Interchange USD: Moved Interchange USD support to its own plugin, which is disabled by default.
This removes USD as a dependency in Interchange.

#virtualized

diff --git a/Engine/Plugins/Interchange/Editor/InterchangeEditor.uplugin b/Engine/Plugins/Interchange/Editor/InterchangeEditor.uplugin
--- a/Engine/Plugins/Interchange/Editor/InterchangeEditor.uplugin
+++ b/Engine/Plugins/Interchange/Editor/InterchangeEditor.uplugin
@@ -56,10 +56,6 @@
 				"Editor",
 				"Program"
 			]
-		},
-		{
-			"Name": "USDCore",
-			"Enabled": true
 		}
 	],
 	"IsExperimentalVersion": false
diff --git a/Engine/Plugins/Interchange/Editor/Source/InterchangeEditor/InterchangeEditor.Build.cs b/Engine/Plugins/Interchange/Editor/Source/InterchangeEditor/InterchangeEditor.Build.cs
--- a/Engine/Plugins/Interchange/Editor/Source/InterchangeEditor/InterchangeEditor.Build.cs
+++ b/Engine/Plugins/Interchange/Editor/Source/InterchangeEditor/InterchangeEditor.Build.cs
@@ -40,8 +40,6 @@ namespace UnrealBuildTool.Rules
 					"InputCore", // Translator settings customizations
 					"Slate", // Translator settings customizations
 					"ToolMenus",
-					"UnrealUSDWrapper", // UnrealIdentifiers in the USD translator settings customization
-					"USDClasses", // USDProjectSettings, to fetch other material purporses
 				}
 			);
 		}
diff --git a/Engine/Plugins/Interchange/Editor/Source/InterchangeEditor/Private/InterchangeEditorModule.cpp b/Engine/Plugins/Interchange/Editor/Source/InterchangeEditor/Private/InterchangeEditorModule.cpp
--- a/Engine/Plugins/Interchange/Editor/Source/InterchangeEditor/Private/InterchangeEditorModule.cpp
+++ b/Engine/Plugins/Interchange/Editor/Source/InterchangeEditor/Private/InterchangeEditorModule.cpp
@@ -7,7 +7,6 @@
 #include "InterchangeEditorLog.h"
 #include "InterchangeFbxAssetImportDataConverter.h"
 #include "InterchangeManager.h"
-#include "InterchangeUsdTranslatorSettingsCustomization.h"
 #include "InterchangeResetContextMenuExtender.h"
 #include "MessageLogModule.h"
 #include "Misc/App.h"
@@ -139,13 +138,6 @@ void FInterchangeEditorModule::StartupModule()
 		};
 
 		InterchangeManager.OnPreDestroyInterchangeManager.AddLambda(UnregisterItems);
-
-		// Translator settings customizations
-		FPropertyEditorModule& PropertyModule = FModuleManager::LoadModuleChecked<FPropertyEditorModule>(TEXT("PropertyEditor"));
-		PropertyModule.RegisterCustomClassLayout(
-			TEXT("InterchangeUsdTranslatorSettings"),
-			FOnGetDetailCustomizationInstance::CreateStatic(&FInterchangeUsdTranslatorSettingsCustomization::MakeInstance)
-		);
 	};
 
 	if (GEngine)
@@ -162,11 +154,6 @@ void FInterchangeEditorModule::StartupModule()
 
 void FInterchangeEditorModule::ShutdownModule()
 {
-	//Translator settings customizations
-	if (FPropertyEditorModule* PropertyModule = FModuleManager::GetModulePtr<FPropertyEditorModule>(TEXT("PropertyEditor")))
-	{
-		PropertyModule->UnregisterCustomClassLayout(TEXT("InterchangeUsdTranslatorSettings"));
-	}
 	FInterchangeResetContextMenuExtender::RemoveLevelEditorContextMenuExtender();
 }
 
diff --git a/Engine/Plugins/Interchange/Editor/Source/InterchangeEditor/Private/InterchangeUsdTranslatorSettingsCustomization.cpp b/Engine/Plugins/Interchange/Editor/Source/InterchangeEditor/Private/InterchangeUsdTranslatorSettingsCustomization.cpp
deleted file mode 100644
diff --git a/Engine/Plugins/Interchange/Editor/Source/InterchangeEditor/Private/InterchangeUsdTranslatorSettingsCustomization.h b/Engine/Plugins/Interchange/Editor/Source/InterchangeEditor/Private/InterchangeUsdTranslatorSettingsCustomization.h
deleted file mode 100644
diff --git a/Engine/Plugins/Interchange/Extensions/OpenUSD/Config/BaseInterchangeOpenUSD.ini b/Engine/Plugins/Interchange/Extensions/OpenUSD/Config/BaseInterchangeOpenUSD.ini
new file mode 100644
--- /dev/null
+++ b/Engine/Plugins/Interchange/Extensions/OpenUSD/Config/BaseInterchangeOpenUSD.ini
@@ -0,0 +1,2 @@
+[CoreRedirects]
++ClassRedirects=(OldName="/Script/InterchangePipelines.InterchangeUsdPipeline",NewName="/Script/InterchangeOpenUSDImport.InterchangeUsdPipeline")
diff --git a/Engine/Plugins/Interchange/Extensions/OpenUSD/InterchangeOpenUSD.uplugin b/Engine/Plugins/Interchange/Extensions/OpenUSD/InterchangeOpenUSD.uplugin
new file mode 100644
--- /dev/null
+++ b/Engine/Plugins/Interchange/Extensions/OpenUSD/InterchangeOpenUSD.uplugin
@@ -0,0 +1,60 @@
+{
+	"FileVersion": 3,
+	"Version": 1,
+	"VersionName": "1.0",
+	"FriendlyName": "Interchange OpenUSD Support",
+	"Description": "Allows translation of OpenUSD files via the Interchange framework",
+	"Category": "Importers",
+	"CreatedBy": "Epic Games, Inc.",
+	"CreatedByURL": "https://epicgames.com",
+	"DocsURL": "",
+	"MarketplaceURL": "",
+	"SupportURL": "",
+	"EnabledByDefault": false,
+	"CanContainContent": true,
+	"IsBetaVersion": false,
+	"Installed": false,
+	"IsExperimentalVersion": true,
+	"SupportedTargetPlatforms": [
+		"Win64",
+		"Linux",
+		"Mac"
+	],
+	"Modules": [
+		{
+			"Name": "InterchangeOpenUSDImport",
+			"Type": "Runtime",
+			"LoadingPhase": "Default"
+		},
+		{
+			"Name": "InterchangeOpenUSDEditor",
+			"Type": "Editor",
+			"LoadingPhase": "Default"
+		}
+	],
+	"Plugins": [
+		{
+			"Name": "Interchange",
+			"Enabled": true
+		},
+		{
+			"Name": "InterchangeEditor",
+			"Enabled": true,
+			"TargetAllowList": [
+				"Editor",
+				"Program"
+			]
+		},
+		{
+			"Name": "InterchangeOpenVDB",
+			"Enabled": true,
+			"TargetAllowList": [
+				"Editor"
+			]
+		},
+		{
+			"Name": "USDCore",
+			"Enabled": true
+		}
+	]
+}
\ No newline at end of file
diff --git a/Engine/Plugins/Interchange/Extensions/OpenUSD/Source/Editor/InterchangeOpenUSDEditor.build.cs b/Engine/Plugins/Interchange/Extensions/OpenUSD/Source/Editor/InterchangeOpenUSDEditor.build.cs
new file mode 100644
--- /dev/null
+++ b/Engine/Plugins/Interchange/Extensions/OpenUSD/Source/Editor/InterchangeOpenUSDEditor.build.cs
@@ -0,0 +1,36 @@
+// Copyright Epic Games, Inc. All Rights Reserved.
+
+using System.IO;
+
+namespace UnrealBuildTool.Rules
+{
+	public class InterchangeOpenUSDEditor : ModuleRules
+	{
+		public InterchangeOpenUSDEditor(ReadOnlyTargetRules Target) : base(Target)
+		{
+			PublicDependencyModuleNames.AddRange(
+				new string[]
+				{
+					"Core",
+					"CoreUObject",
+					"Engine"
+				}
+			);
+
+			PrivateDependencyModuleNames.AddRange(
+				new string[]
+				{
+					"InputCore",			// For EKeys
+					"InterchangeEditor",
+					"InterchangeEngine",	// For UInterchangeManager::GetInterchangeManager
+					"InterchangeOpenUSDImport",
+					"Slate",
+					"SlateCore",
+					"UnrealEd",				// For FScopedTransaction
+					"UnrealUSDWrapper",		// UnrealIdentifiers in the USD translator settings customization
+					"USDClasses"			// For USDProjectSettings and render context
+				}
+			);
+		}
+	}
+}
\ No newline at end of file
diff --git a/Engine/Plugins/Interchange/Extensions/OpenUSD/Source/Editor/Private/InterchangeOpenUSDEditorModule.cpp b/Engine/Plugins/Interchange/Extensions/OpenUSD/Source/Editor/Private/InterchangeOpenUSDEditorModule.cpp
new file mode 100644
--- /dev/null
+++ b/Engine/Plugins/Interchange/Extensions/OpenUSD/Source/Editor/Private/InterchangeOpenUSDEditorModule.cpp
@@ -0,0 +1,42 @@
+// Copyright Epic Games, Inc. All Rights Reserved.
+
+#include "InterchangeOpenUSDEditorModule.h"
+
+#include "Engine/Engine.h"
+#include "InterchangeManager.h"
+#include "InterchangeUsdTranslatorSettingsCustomization.h"
+#include "Misc/CoreDelegates.h"
+#include "Modules/ModuleManager.h"
+#include "PropertyEditorModule.h"
+
+void FInterchangeOpenUSDEditorModule::StartupModule()
+{
+	using namespace UE::Interchange;
+	auto RegisterItems = []()
+	{
+		// Translator settings customizations
+		FPropertyEditorModule& PropertyModule = FModuleManager::LoadModuleChecked<FPropertyEditorModule>(TEXT("PropertyEditor"));
+		PropertyModule.RegisterCustomClassLayout(
+			TEXT("InterchangeUsdTranslatorSettings"),
+			FOnGetDetailCustomizationInstance::CreateStatic(&FInterchangeUsdTranslatorSettingsCustomization::MakeInstance)
+		);
+	};
+	if (GEngine)
+	{
+		RegisterItems();
+	}
+	else
+	{
+		FCoreDelegates::OnPostEngineInit.AddLambda(RegisterItems);
+	}
+}
+void FInterchangeOpenUSDEditorModule::ShutdownModule()
+{
+	//Translator settings customizations
+	if (FPropertyEditorModule* PropertyModule = FModuleManager::GetModulePtr<FPropertyEditorModule>(TEXT("PropertyEditor")))
+	{
+		PropertyModule->UnregisterCustomClassLayout(TEXT("InterchangeUsdTranslatorSettings"));
+	}
+}
+
+IMPLEMENT_MODULE(FInterchangeOpenUSDEditorModule, InterchangeOpenUSDEditor)
\ No newline at end of file
diff --git a/Engine/Plugins/Interchange/Extensions/OpenUSD/Source/Editor/Private/InterchangeUsdTranslatorSettingsCustomization.cpp b/Engine/Plugins/Interchange/Extensions/OpenUSD/Source/Editor/Private/InterchangeUsdTranslatorSettingsCustomization.cpp
new file mode 100644
--- /dev/null
+++ b/Engine/Plugins/Interchange/Extensions/OpenUSD/Source/Editor/Private/InterchangeUsdTranslatorSettingsCustomization.cpp
@@ -0,0 +1,302 @@
+// Copyright Epic Games, Inc. All Rights Reserved.
+
+#include "InterchangeUsdTranslatorSettingsCustomization.h"
+
+#include "InterchangeUsdTranslator.h"
+#include "UnrealUSDWrapper.h"
+#include "USDMaterialUtils.h"
+#include "USDProjectSettings.h"
+
+#include "DetailCategoryBuilder.h"
+#include "DetailLayoutBuilder.h"
+#include "DetailWidgetRow.h"
+#include "Modules/ModuleManager.h"
+#include "ScopedTransaction.h"
+#include "Styling/AppStyle.h"
+#include "UObject/Object.h"
+#include "Widgets/DeclarativeSyntaxSupport.h"
+#include "Widgets/Input/SEditableTextBox.h"
+
+#define LOCTEXT_NAMESPACE "InterchangeUsdTranslatorSettingsCustomization"
+
+FInterchangeUsdTranslatorSettingsCustomization::FInterchangeUsdTranslatorSettingsCustomization()
+{
+}
+
+TSharedRef<IDetailCustomization> FInterchangeUsdTranslatorSettingsCustomization::MakeInstance()
+{
+	return MakeShared<FInterchangeUsdTranslatorSettingsCustomization>();
+}
+
+void FInterchangeUsdTranslatorSettingsCustomization::CustomizeDetails(IDetailLayoutBuilder& DetailLayoutBuilder)
+{
+	TArray<TWeakObjectPtr<UObject>> SelectedObjects = DetailLayoutBuilder.GetSelectedObjects();
+	if (SelectedObjects.Num() != 1)
+	{
+		return;
+	}
+
+	TWeakObjectPtr<UObject> SelectedObject = SelectedObjects[0];
+	if (!SelectedObject.IsValid())
+	{
+		return;
+	}
+
+	CurrentOptions = Cast<UInterchangeUsdTranslatorSettings>(SelectedObject.Get());
+	if (!CurrentOptions)
+	{
+		return;
+	}
+
+	RenderContextComboBoxItems.Reset();
+	TSharedPtr<FString> InitiallySelectedContext;
+	for (const FName& Context : UsdUnreal::MaterialUtils::GetRegisteredRenderContexts())
+	{
+		TSharedPtr<FString> ContextStr;
+		if (Context == UnrealIdentifiers::UniversalRenderContext)
+		{
+			ContextStr = MakeShared<FString>(UnrealIdentifiers::UniversalRenderContextDisplayString);
+		}
+		else
+		{
+			ContextStr = MakeShared<FString>(Context.ToString());
+		}
+
+		if (Context == CurrentOptions->RenderContext)
+		{
+			InitiallySelectedContext = ContextStr;
+		}
+
+		RenderContextComboBoxItems.Add(ContextStr);
+	}
+
+	IDetailCategoryBuilder& CatBuilder = DetailLayoutBuilder.EditCategory(TEXT("USD Translator"));
+
+	if (TSharedPtr<IPropertyHandle> RenderContextProperty = DetailLayoutBuilder.GetProperty(
+			GET_MEMBER_NAME_CHECKED(UInterchangeUsdTranslatorSettings, RenderContext)
+		))
+	{
+		DetailLayoutBuilder.HideProperty(RenderContextProperty);
+
+		// clang-format off
+		CatBuilder.AddCustomRow(FText::FromString(TEXT("RenderContextCustomization")))
+		.NameContent()
+		[
+			SNew(STextBlock)
+			.Text(FText::FromString(TEXT("Render Context")))
+			.Font(FAppStyle::GetFontStyle(TEXT("PropertyWindow.NormalFont")))
+			.ToolTipText(RenderContextProperty->GetToolTipText())
+		]
+		.ValueContent()
+		[
+			SAssignNew(RenderContextComboBox, SComboBox<TSharedPtr<FString>>)
+			.OptionsSource(&RenderContextComboBoxItems)
+			.InitiallySelectedItem(InitiallySelectedContext)
+			.OnSelectionChanged_Lambda([this, RenderContextProperty](TSharedPtr<FString> NewContext, ESelectInfo::Type SelectType)
+			{
+				if (CurrentOptions == nullptr || !NewContext.IsValid())
+				{
+					return;
+				}
+
+				FScopedTransaction Transaction(LOCTEXT("RenderContextTransaction", "Edit Render Context"));
+				RenderContextProperty->NotifyPreChange();
+				{
+					FName NewContextName = (*NewContext) == UnrealIdentifiers::UniversalRenderContextDisplayString
+						? UnrealIdentifiers::UniversalRenderContext
+						: FName(**NewContext);
+
+					CurrentOptions->RenderContext = NewContextName;
+				}
+				RenderContextProperty->NotifyPostChange(EPropertyChangeType::ValueSet);
+				RenderContextProperty->NotifyFinishedChangingProperties();
+			})
+			.OnGenerateWidget_Lambda([](TSharedPtr<FString> Item)
+			{
+				return SNew(STextBlock)
+					.Text(Item.IsValid() ? FText::FromString(*Item) : FText::GetEmpty())
+					.Font(FAppStyle::GetFontStyle(TEXT("PropertyWindow.NormalFont")));
+			})
+			.Content()
+			[
+				SNew(STextBlock)
+				.Text_Lambda([this]() -> FText
+				{
+					TSharedPtr<FString> SelectedItem = RenderContextComboBox->GetSelectedItem();
+					if (SelectedItem.IsValid())
+					{
+						return FText::FromString(*SelectedItem);
+					}
+
+					return FText::GetEmpty();
+				})
+				.Font(FAppStyle::GetFontStyle(TEXT("PropertyWindow.NormalFont")))
+			]
+		];
+		// clang-format on
+	}
+
+	if (TSharedPtr<IPropertyHandle> MaterialPurposeProperty = DetailLayoutBuilder.GetProperty(
+			GET_MEMBER_NAME_CHECKED(UInterchangeUsdTranslatorSettings, MaterialPurpose)
+		))
+	{
+		DetailLayoutBuilder.HideProperty(MaterialPurposeProperty);
+
+		// clang-format off
+		CatBuilder.AddCustomRow(FText::FromString(TEXT("MaterialPurposeCustomization")))
+		.NameContent()
+		[
+			SNew(STextBlock)
+			.Text(FText::FromString(TEXT("Material purpose")))
+			.Font(FAppStyle::GetFontStyle(TEXT("PropertyWindow.NormalFont")))
+			.ToolTipText(MaterialPurposeProperty->GetToolTipText())
+		]
+		.ValueContent()
+		[
+			SNew(SBox)
+			.VAlign(VAlign_Center)
+			[
+				SNew(SComboBox<TSharedPtr<FString>>)
+				.OptionsSource(&MaterialPurposeComboBoxItems)
+				.OnComboBoxOpening_Lambda([this]()
+				{
+					MaterialPurposeComboBoxItems = {
+						MakeShared<FString>(UnrealIdentifiers::MaterialAllPurpose),
+						MakeShared<FString>(UnrealIdentifiers::MaterialPreviewPurpose),
+						MakeShared<FString>(UnrealIdentifiers::MaterialFullPurpose)
+					};
+
+					// Add additional purposes from project settings
+					if (const UUsdProjectSettings* ProjectSettings = GetDefault<UUsdProjectSettings>())
+					{
+						MaterialPurposeComboBoxItems.Reserve(MaterialPurposeComboBoxItems.Num() + ProjectSettings->AdditionalMaterialPurposes.Num());
+
+						TSet<FString> ExistingEntries = {
+							UnrealIdentifiers::MaterialAllPurpose,
+							UnrealIdentifiers::MaterialPreviewPurpose,
+							UnrealIdentifiers::MaterialFullPurpose
+						};
+
+						for (const FName& AdditionalPurpose : ProjectSettings->AdditionalMaterialPurposes)
+						{
+							FString AdditionalPurposeStr = AdditionalPurpose.ToString();
+
+							if (!ExistingEntries.Contains(AdditionalPurposeStr))
+							{
+								ExistingEntries.Add(AdditionalPurposeStr);
+								MaterialPurposeComboBoxItems.AddUnique(MakeShared<FString>(AdditionalPurposeStr));
+							}
+						}
+					}
+				})
+				.OnGenerateWidget_Lambda([&](TSharedPtr<FString> Option)
+				{
+					TSharedPtr<SWidget> Widget = SNullWidget::NullWidget;
+					if (Option)
+					{
+						Widget = SNew(STextBlock)
+							.Text(FText::FromString((*Option) == UnrealIdentifiers::MaterialAllPurpose
+								? UnrealIdentifiers::MaterialAllPurposeText
+								: *Option
+							))
+							.Font(FAppStyle::GetFontStyle("PropertyWindow.NormalFont"));
+					}
+
+					return Widget.ToSharedRef();
+				})
+				.OnSelectionChanged_Lambda([this, MaterialPurposeProperty](TSharedPtr<FString> ChosenOption, ESelectInfo::Type SelectInfo)
+				{
+					if (CurrentOptions && ChosenOption)
+					{
+						FScopedTransaction Transaction(LOCTEXT("MaterialPurposeTransaction", "Edit Material Purpose"));
+						MaterialPurposeProperty->NotifyPreChange();
+						{
+							CurrentOptions->MaterialPurpose = **ChosenOption;
+						}
+						MaterialPurposeProperty->NotifyPostChange(EPropertyChangeType::ValueSet);
+						MaterialPurposeProperty->NotifyFinishedChangingProperties();
+					}
+
+				})
+				[
+					SNew(SEditableTextBox)
+					.Text_Lambda([this]() -> FText
+					{
+						if (CurrentOptions)
+						{
+							return FText::FromString(CurrentOptions->MaterialPurpose == *UnrealIdentifiers::MaterialAllPurpose
+								? UnrealIdentifiers::MaterialAllPurposeText
+								: CurrentOptions->MaterialPurpose.ToString()
+							);
+						}
+
+						return FText::GetEmpty();
+					})
+					.Font(FAppStyle::GetFontStyle("PropertyWindow.NormalFont"))
+					.OnTextCommitted_Lambda([this, MaterialPurposeProperty](const FText& NewText, ETextCommit::Type CommitType)
+					{
+						if (CommitType != ETextCommit::OnEnter)
+						{
+							return;
+						}
+
+						FString NewPurposeString = NewText.ToString();
+						FName NewPurpose = *NewPurposeString;
+
+						bool bIsNew = true;
+						for (const TSharedPtr<FString>& Purpose : MaterialPurposeComboBoxItems)
+						{
+							if (Purpose && *Purpose == NewPurposeString)
+							{
+								bIsNew = false;
+								break;
+							}
+						}
+
+						if (bIsNew)
+						{
+							if (UUsdProjectSettings* ProjectSettings = GetMutableDefault<UUsdProjectSettings>())
+							{
+								ProjectSettings->AdditionalMaterialPurposes.AddUnique(NewPurpose);
+								ProjectSettings->SaveConfig();
+							}
+						}
+
+						if (CurrentOptions)
+						{
+							FScopedTransaction Transaction(LOCTEXT("MaterialPurposeTypeTransaction", "Add and Set Material Purpose"));
+							MaterialPurposeProperty->NotifyPreChange();
+							{
+								CurrentOptions->MaterialPurpose = NewPurpose;
+							}
+							MaterialPurposeProperty->NotifyPostChange(EPropertyChangeType::ValueSet);
+							MaterialPurposeProperty->NotifyFinishedChangingProperties();
+						}
+					})
+				]
+			]
+		];
+		// clang-format on
+	}
+
+	// Add/remove properties so that they retain their usual order
+	if (TSharedPtr<IPropertyHandle> OverrideStageOptionsProperty = DetailLayoutBuilder.GetProperty(
+			GET_MEMBER_NAME_CHECKED(UInterchangeUsdTranslatorSettings, bOverrideStageOptions)
+		))
+	{
+		CatBuilder.AddProperty(OverrideStageOptionsProperty);
+	}
+	if (TSharedPtr<IPropertyHandle> StageOptionsProperty = DetailLayoutBuilder.GetProperty(
+			GET_MEMBER_NAME_CHECKED(UInterchangeUsdTranslatorSettings, StageOptions)
+		))
+	{
+		CatBuilder.AddProperty(StageOptionsProperty);
+	}
+}
+
+void FInterchangeUsdTranslatorSettingsCustomization::CustomizeDetails(const TSharedPtr<IDetailLayoutBuilder>& DetailBuilder)
+{
+	CustomizeDetails(*DetailBuilder);
+}
+
+#undef LOCTEXT_NAMESPACE
diff --git a/Engine/Plugins/Interchange/Extensions/OpenUSD/Source/Editor/Private/InterchangeUsdTranslatorSettingsCustomization.h b/Engine/Plugins/Interchange/Extensions/OpenUSD/Source/Editor/Private/InterchangeUsdTranslatorSettingsCustomization.h
new file mode 100644
--- /dev/null
+++ b/Engine/Plugins/Interchange/Extensions/OpenUSD/Source/Editor/Private/InterchangeUsdTranslatorSettingsCustomization.h
@@ -0,0 +1,33 @@
+// Copyright Epic Games, Inc. All Rights Reserved.
+
+#pragma once
+
+#if WITH_EDITOR
+
+#include "IDetailCustomization.h"
+#include "Widgets/Input/SComboBox.h"
+
+class UInterchangeUsdTranslatorSettings;
+class IDetailLayoutBuilder;
+
+class FInterchangeUsdTranslatorSettingsCustomization : public IDetailCustomization
+{
+public:
+	FInterchangeUsdTranslatorSettingsCustomization();
+	static TSharedRef<IDetailCustomization> MakeInstance();
+
+	// IDetailCustomization interface
+	virtual void CustomizeDetails(IDetailLayoutBuilder& DetailLayout) override;
+	virtual void CustomizeDetails(const TSharedPtr<IDetailLayoutBuilder>& DetailBuilder) override;
+	// End of IDetailCustomization interface
+
+private:
+	UInterchangeUsdTranslatorSettings* CurrentOptions;
+
+	TSharedPtr<SComboBox<TSharedPtr<FString>>> RenderContextComboBox;
+	TArray<TSharedPtr<FString>> RenderContextComboBoxItems;
+
+	TArray<TSharedPtr<FString>> MaterialPurposeComboBoxItems;
+};
+
+#endif	  // WITH_EDITOR
diff --git a/Engine/Plugins/Interchange/Extensions/OpenUSD/Source/Editor/Public/InterchangeOpenUSDEditorModule.h b/Engine/Plugins/Interchange/Extensions/OpenUSD/Source/Editor/Public/InterchangeOpenUSDEditorModule.h
new file mode 100644
--- /dev/null
+++ b/Engine/Plugins/Interchange/Extensions/OpenUSD/Source/Editor/Public/InterchangeOpenUSDEditorModule.h
@@ -0,0 +1,12 @@
+// Copyright Epic Games, Inc. All Rights Reserved.
+
+#pragma once
+
+#include "Modules/ModuleInterface.h"
+
+class FInterchangeOpenUSDEditorModule : public IModuleInterface
+{
+public:
+	virtual void StartupModule() override;
+	virtual void ShutdownModule() override;
+};
\ No newline at end of file
diff --git a/Engine/Plugins/Interchange/Extensions/OpenUSD/Source/Import/InterchangeOpenUSDImport.build.cs b/Engine/Plugins/Interchange/Extensions/OpenUSD/Source/Import/InterchangeOpenUSDImport.build.cs
new file mode 100644
--- /dev/null
+++ b/Engine/Plugins/Interchange/Extensions/OpenUSD/Source/Import/InterchangeOpenUSDImport.build.cs
@@ -0,0 +1,60 @@
+// Copyright Epic Games, Inc. All Rights Reserved.
+
+using System.IO;
+
+namespace UnrealBuildTool.Rules
+{
+	public class InterchangeOpenUSDImport : ModuleRules
+	{
+		public InterchangeOpenUSDImport(ReadOnlyTargetRules Target) : base(Target)
+		{
+			PublicDependencyModuleNames.AddRange(
+				new string[]
+				{
+					"Core",
+					"CoreUObject",
+					"Engine",
+					"InterchangeCommonParser",	// For FAnimationPayloadData
+					"InterchangeCore", 			// For UInterchangeBaseNode, base translator
+					"InterchangeImport", 		// For the translator settings, pipeline interface, etc.
+					"UnrealUSDWrapper",
+					"USDUtilities"
+				}
+			);
+
+			PrivateDependencyModuleNames.AddRange(
+				new string[]
+				{
+					"DynamicMesh",				// For FMeshMapMaker
+					"GeometryCore",				// For FDynamicMesh3
+					"InterchangeEngine", 		// For UInterchangeManager::GetInterchangeManager
+					"InterchangeFactoryNodes",
+					"InterchangeNodes", 		// For UInterchangeShaderGraphNode
+					"InterchangePipelines",
+					"MeshConversion",			// For FMeshDescriptionToDynamicMesh
+					"MeshDescription",
+					"StaticMeshDescription",	// FStaticMeshOperations
+					"TextureUtilitiesCommon",	// For UDIM functions
+					"USDClasses"
+				}
+			);
+
+			if (Target.Type == TargetType.Editor)
+			{
+				PrivateDependencyModuleNames.AddRange(
+					new string[]
+					{
+						"InterchangeOpenVDBImport",
+						"ModelingComponents",	// For FTexture2DBuilder
+					}
+				);
+
+				if (Target.Platform == UnrealTargetPlatform.Win64 || Target.LinkType == TargetLinkType.Monolithic ||
+					(Target.Platform == UnrealTargetPlatform.Mac && Target.Architecture == UnrealArch.X64))
+				{
+					PrivateDependencyModuleNames.Add("MaterialX");
+				}
+			}
+		}
+	}
+}
\ No newline at end of file
diff --git a/Engine/Plugins/Interchange/Extensions/OpenUSD/Source/Import/Private/InterchangeOpenUSDImportModule.cpp b/Engine/Plugins/Interchange/Extensions/OpenUSD/Source/Import/Private/InterchangeOpenUSDImportModule.cpp
new file mode 100644
--- /dev/null
+++ b/Engine/Plugins/Interchange/Extensions/OpenUSD/Source/Import/Private/InterchangeOpenUSDImportModule.cpp
@@ -0,0 +1,32 @@
+// Copyright Epic Games, Inc. All Rights Reserved.
+
+#include "InterchangeOpenUSDImportModule.h"
+
+#include "Engine/Engine.h"
+#include "InterchangeManager.h"
+#include "InterchangeUsdTranslator.h"
+#include "Misc/CoreDelegates.h"
+#include "Modules/ModuleManager.h"
+
+void FInterchangeOpenUSDImportModule::StartupModule()
+{
+	using namespace UE::Interchange;
+	auto RegisterItems = []()
+	{
+		UInterchangeManager& InterchangeManager = UInterchangeManager::GetInterchangeManager();
+		InterchangeManager.RegisterTranslator(UInterchangeUSDTranslator::StaticClass());
+	};
+	if (GEngine)
+	{
+		RegisterItems();
+	}
+	else
+	{
+		FCoreDelegates::OnPostEngineInit.AddLambda(RegisterItems);
+	}
+}
+void FInterchangeOpenUSDImportModule::ShutdownModule()
+{
+}
+
+IMPLEMENT_MODULE(FInterchangeOpenUSDImportModule, InterchangeOpenUSDImport)
\ No newline at end of file
diff --git a/Engine/Plugins/Interchange/Extensions/OpenUSD/Source/Import/Private/InterchangeUSDPipeline.cpp b/Engine/Plugins/Interchange/Extensions/OpenUSD/Source/Import/Private/InterchangeUSDPipeline.cpp
new file mode 100644
--- /dev/null
+++ b/Engine/Plugins/Interchange/Extensions/OpenUSD/Source/Import/Private/InterchangeUSDPipeline.cpp
@@ -0,0 +1,1495 @@
+// Copyright Epic Games, Inc. All Rights Reserved.
+
+#include "InterchangeUSDPipeline.h"
+
+#include "InterchangeHeterogeneousVolumeActorFactoryNode.h"
+#include "InterchangeManager.h"
+#include "InterchangeMaterialDefinitions.h"
+#include "InterchangeMaterialFactoryNode.h"
+#include "InterchangeMeshFactoryNode.h"
+#include "InterchangeMeshNode.h"
+#include "InterchangeSceneNode.h"
+#include "InterchangeShaderGraphNode.h"
+#include "InterchangeSkeletalMeshLodDataNode.h"
+#include "InterchangeSparseVolumeTextureFactoryNode.h"
+#include "InterchangeStaticMeshLodDataNode.h"
+#include "InterchangeVolumeNode.h"
+#include "Nodes/InterchangeBaseNodeContainer.h"
+#include "Nodes/InterchangeUserDefinedAttribute.h"
+#include "Usd/InterchangeUsdDefinitions.h"
+
+#include "DynamicMesh/DynamicMesh3.h"
+#include "Engine/StaticMesh.h"
+#include "Materials/Material.h"
+#include "Materials/MaterialExpressionTextureBase.h"
+#include "Materials/MaterialExpressionTextureSample.h"
+#include "Materials/MaterialInstanceConstant.h"
+#include "Materials/MaterialInstanceDynamic.h"
+#include "MeshDescriptionToDynamicMesh.h"
+#include "Sampling/MeshMapBaker.h"
+#include "Sampling/MeshPropertyMapEvaluator.h"
+#include "StaticMeshAttributes.h"
+#include "StaticMeshOperations.h"
+
+#include "UnrealUSDWrapper.h"
+#include "Usd/InterchangeUsdDefinitions.h"
+#include "USDConversionUtils.h"
+#include "USDLog.h"
+#include "USDMaterialUtils.h"
+#include "USDProjectSettings.h"
+#include "UsdWrappers/SdfPath.h"
+
+#if WITH_EDITORONLY_DATA
+#include "AssetUtils/Texture2DBuilder.h"
+#endif	  // WITH_EDITORONLY_DATA
+
+#include UE_INLINE_GENERATED_CPP_BY_NAME(InterchangeUSDPipeline)
+
+DEFINE_LOG_CATEGORY_STATIC(LogInterchangeUSDPipeline, Log, All);
+
+namespace UE::InterchangeUsdPipeline::Private
+{
+#if USE_USD_SDK
+	const FString& GetPseudoRootTranslatedNodeUid()
+	{
+		const static FString Result = UE::FSdfPath::AbsoluteRootPath().GetString();
+		return Result;
+	}
+
+	const FString& GetPseudoRootFactoryNodeUid()
+	{
+		const static FString Result = UInterchangeFactoryBaseNode::BuildFactoryNodeUid(GetPseudoRootTranslatedNodeUid());
+		return Result;
+	}
+
+	// Disables the factory nodes on NodeContainer if they target translated nodes that had disallowed purposes
+	void DisableNodesBasedOnGeometryPurpose(UInterchangeBaseNodeContainer& NodeContainer, const EUsdPurpose AllowedPurposes)
+	{
+		TSet<FString> DisabledTranslatedNodes;
+
+		TFunction<void(const FString&, bool)> RecursiveCollectFilteredTranslatedNodes = nullptr;
+		RecursiveCollectFilteredTranslatedNodes = [&NodeContainer,
+												   AllowedPurposes,
+												   &RecursiveCollectFilteredTranslatedNodes,
+												   &DisabledTranslatedNodes](const FString& NodeUid, bool bDisableSubtree)
+		{
+			const UInterchangeBaseNode* Node = NodeContainer.GetNode(NodeUid);
+			if (!Node)
+			{
+				return;
+			}
+
+			if (!bDisableSubtree)
+			{
+				EUsdPurpose AuthoredPurpose = EUsdPurpose::Default;
+				bool bHasPurpose = Node->GetInt32Attribute(UE::Interchange::USD::GeometryPurposeIdentifier, (int32&)AuthoredPurpose);
+				if (bHasPurpose && !EnumHasAllFlags(AllowedPurposes, AuthoredPurpose))
+				{
+					// Purpose inheritance according to UsdGeomImageable::ComputePurposeInfo seems to be:
+					// - Authored purposes inherit down to prims without any purpose;
+					// - If a prim has any purpose authored on them, that is their "computed purpose", and that is inherited to its children.
+					//
+					// We don't care much about the actual purpose then: We just need to check if any prim/node has a purpose
+					// that is not allowed. If that is the case, we turn off the entire subtree
+					bDisableSubtree = true;
+				}
+			}
+
+			if (bDisableSubtree)
+			{
+				// Mark the scene translated node as disabled
+				DisabledTranslatedNodes.Add(NodeUid);
+
+				// Mark the asset translated node as disabled
+				if (const UInterchangeSceneNode* SceneNode = Cast<UInterchangeSceneNode>(Node))
+				{
+					FString AssetNodeUid;
+					if (SceneNode->GetCustomAssetInstanceUid(AssetNodeUid))
+					{
+						DisabledTranslatedNodes.Add(AssetNodeUid);
+					}
+				}
+			}
+
+			for (const FString& ChildNodeUid : NodeContainer.GetNodeChildrenUids(NodeUid))
+			{
+				RecursiveCollectFilteredTranslatedNodes(ChildNodeUid, bDisableSubtree);
+			}
+		};
+		const bool bDisableSubtree = false;
+		RecursiveCollectFilteredTranslatedNodes(GetPseudoRootTranslatedNodeUid(), bDisableSubtree);
+
+		// Disable any factory node that targets/references one of our disabled scene nodes.
+		//
+		// We do this in a separate pass because we can say absolutely nothing about our factory nodes, as any number
+		// of pipelines may have done arbirary transformations before we got to run. We have to hope they kept the
+		// target node attribute updated at least, and use that.
+		NodeContainer.IterateNodesOfType<UInterchangeFactoryBaseNode>(
+			[&DisabledTranslatedNodes](const FString& NodeUid, UInterchangeFactoryBaseNode* FactoryNode)
+			{
+				// Handle standard factory nodes
+				{
+					TArray<FString> TargetTranslatedNodeUids;
+					FactoryNode->GetTargetNodeUids(TargetTranslatedNodeUids);
+
+					// If we have a single target node that is disabled, just disable the factory node itself.
+					// This is probably the general case, and should be more explicit on preview panels and such
+					if (TargetTranslatedNodeUids.Num() == 1)
+					{
+						if (DisabledTranslatedNodes.Contains(TargetTranslatedNodeUids[0]))
+						{
+							const bool bEnabled = false;
+							FactoryNode->SetEnabled(bEnabled);
+						}
+					}
+					// If we have multiple target nodes, instead of disabling the whole factory node just remove
+					// this one particular target (e.g. removing a single mesh from a collapsed mesh blob, in case that
+					// mesh doesn't match the purpose)
+					else
+					{
+						for (const FString& TargetNode : TargetTranslatedNodeUids)
+						{
+							if (DisabledTranslatedNodes.Contains(TargetNode))
+							{
+								FactoryNode->RemoveTargetNodeUid(TargetNode);
+							}
+						}
+					}
+				}
+
+				// These nodes don't use the common "target" mechanism and store their mesh nodes separately...
+				// TODO: Will have to expand these to also handle geometry caches?
+				if (UInterchangeStaticMeshLodDataNode* LodDataNode = Cast<UInterchangeStaticMeshLodDataNode>(FactoryNode))
+				{
+					TArray<FString> TranslatedMeshNodeUids;
+					LodDataNode->GetMeshUids(TranslatedMeshNodeUids);
+
+					if (TranslatedMeshNodeUids.Num() == 1)
+					{
+						if (DisabledTranslatedNodes.Contains(TranslatedMeshNodeUids[0]))
+						{
+							const bool bEnabled = false;
+							LodDataNode->SetEnabled(bEnabled);
+						}
+					}
+					else
+					{
+						for (const FString& MeshNodeUid : TranslatedMeshNodeUids)
+						{
+							if (DisabledTranslatedNodes.Contains(MeshNodeUid))
+							{
+								LodDataNode->RemoveMeshUid(MeshNodeUid);
+							}
+						}
+					}
+				}
+				else if (UInterchangeSkeletalMeshLodDataNode* SkeletalNode = Cast<UInterchangeSkeletalMeshLodDataNode>(FactoryNode))
+				{
+					TArray<FString> TranslatedMeshNodeUids;
+					SkeletalNode->GetMeshUids(TranslatedMeshNodeUids);
+
+					if (TranslatedMeshNodeUids.Num() == 1)
+					{
+						if (DisabledTranslatedNodes.Contains(TranslatedMeshNodeUids[0]))
+						{
+							const bool bEnabled = false;
+							SkeletalNode->SetEnabled(bEnabled);
+						}
+					}
+					else
+					{
+						for (const FString& MeshNodeUid : TranslatedMeshNodeUids)
+						{
+							if (DisabledTranslatedNodes.Contains(MeshNodeUid))
+							{
+								SkeletalNode->RemoveMeshUid(MeshNodeUid);
+							}
+						}
+					}
+				}
+			}
+		);
+	}
+
+	// Resolves which should be the parent material for NewFactoryNode and sets it
+	void SetCustomParentMaterial(const UInterchangeShaderGraphNode& ShaderGraphNode, UInterchangeMaterialInstanceFactoryNode& NewFactoryNode)
+	{
+		using namespace UsdUnreal::MaterialUtils;
+
+		// Priority 1) Explicitly set parent material
+		//
+		// This can be from an UnrealMaterial or from anything the user may have manually set
+		// TODO: Maybe use this to handle UnrealMaterials?
+		if (FString ParentMaterialPath;
+			ShaderGraphNode.GetStringAttribute(UE::Interchange::USD::ParentMaterialIdentifier, ParentMaterialPath) && !ParentMaterialPath.IsEmpty())
+		{
+			NewFactoryNode.SetCustomParent(ParentMaterialPath);
+			return;
+		}
+
+		// Priority 2) DisplayColor material
+		//
+		// This is usually a lower priority, but we explicitly flag these so it's easy to check for them
+		if (FString DisplayColorDesc;
+			ShaderGraphNode.GetStringAttribute(UE::Interchange::USD::DisplayColorDescIdentifier, DisplayColorDesc) && !DisplayColorDesc.IsEmpty())
+		{
+			TOptional<FDisplayColorMaterial> ParsedMat = FDisplayColorMaterial::FromString(DisplayColorDesc);
+			if (ParsedMat.IsSet())
+			{
+				if (const FSoftObjectPath* ParentMaterialPath = GetReferenceMaterialPath(ParsedMat.GetValue()))
+				{
+					NewFactoryNode.SetCustomParent(ParentMaterialPath->GetAssetPathString());
+					return;
+				}
+			}
+		}
+
+		// Priority 3) UsdPreviewSurface material
+		{
+			EUsdReferenceMaterialProperties Properties = EUsdReferenceMaterialProperties::None;
+
+			// Translucent
+			int BlendMode = EBlendMode::BLEND_Opaque;
+			if (ShaderGraphNode.GetCustomBlendMode(BlendMode) && BlendMode == EBlendMode::BLEND_Translucent)
+			{
+				Properties |= EUsdReferenceMaterialProperties::Translucent;
+			}
+
+			// TwoSided
+			// TODO: Actually set this somewhere
+			bool bTwoSided = false;
+			if (ShaderGraphNode.GetCustomTwoSided(bTwoSided) && bTwoSided)
+			{
+				Properties |= EUsdReferenceMaterialProperties::TwoSided;
+			}
+
+			// VT
+			bool bNeedsVT = false;
+			if (ShaderGraphNode.GetBooleanAttribute(UE::Interchange::USD::NeedsVTParentIdentifier, bNeedsVT) && bNeedsVT)
+			{
+				Properties |= EUsdReferenceMaterialProperties::VT;
+			}
+
+			FSoftObjectPath ParentMaterial = GetReferencePreviewSurfaceMaterial(Properties);
+			if (ParentMaterial.IsValid())
+			{
+				NewFactoryNode.SetCustomParent(ParentMaterial.GetAssetPathString());
+				return;
+			}
+		}
+	}
+
+	// Creates a new material instance factory node and fills it in with the parameter/parent material path
+	// stored as base translated node attributes on the provided ShaderGraphNode
+	UInterchangeMaterialInstanceFactoryNode* BuildMaterialInstance(
+		UInterchangeBaseNodeContainer& NodeContainer,
+		const UInterchangeShaderGraphNode& ShaderGraphNode
+	)
+	{
+		const FString ShaderGraphNodeUid = ShaderGraphNode.GetUniqueID();
+		FString MaterialFactoryNodeUid = UInterchangeFactoryBaseNode::BuildFactoryNodeUid(ShaderGraphNodeUid);
+		FString MaterialFactoryNodeName = ShaderGraphNode.GetDisplayLabel();
+
+		UInterchangeMaterialInstanceFactoryNode* NewFactoryNode = NewObject<UInterchangeMaterialInstanceFactoryNode>(&NodeContainer);
+		NewFactoryNode->InitializeNode(MaterialFactoryNodeUid, MaterialFactoryNodeName, EInterchangeNodeContainerType::FactoryData);
+		NewFactoryNode->AddTargetNodeUid(ShaderGraphNodeUid);
+
+		const bool bAddSourceNodeName = false;
+		UInterchangeUserDefinedAttributesAPI::DuplicateAllUserDefinedAttribute(&ShaderGraphNode, NewFactoryNode, bAddSourceNodeName);
+
+		SetCustomParentMaterial(ShaderGraphNode, *NewFactoryNode);
+
+		// Pick material instance class
+		{
+			UInterchangeEditorUtilitiesBase* EditorUtilities = UInterchangeManager::GetInterchangeManager().GetEditorUtilities();
+			const UClass* MaterialClass = (EditorUtilities && EditorUtilities->IsRuntimeOrPIE()) ? UMaterialInstanceDynamic::StaticClass()
+																								 : UMaterialInstanceConstant::StaticClass();
+			NewFactoryNode->SetCustomInstanceClassName(MaterialClass->GetPathName());
+		}
+
+		// Move material attributes to the factory nodes
+		TArray<UE::Interchange::FAttributeKey> AllAttributeKeys;
+		ShaderGraphNode.GetAttributeKeys(AllAttributeKeys);
+		for (const UE::Interchange::FAttributeKey& AttributeKey : AllAttributeKeys)
+		{
+			const FString AttributeKeyString = AttributeKey.ToString();
+
+			// It's a regular material attribute
+			if (AttributeKeyString.Contains(UE::Interchange::USD::AttributeIdentifierPrefix))
+			{
+				const FString CleanAttributeKey = AttributeKeyString.Replace(	 //
+					*UE::Interchange::USD::AttributeIdentifierPrefix,
+					TEXT(""),
+					ESearchCase::CaseSensitive
+				);
+				const FString StandardParameterName = UInterchangeShaderPortsAPI::MakeInputValueKey(CleanAttributeKey);
+
+				const UE::Interchange::EAttributeTypes AttributeType = ShaderGraphNode.GetAttributeType(AttributeKey);
+				switch (AttributeType)
+				{
+					case UE::Interchange::EAttributeTypes::Bool:
+						if (bool Value; ShaderGraphNode.GetBooleanAttribute(AttributeKey.Key, Value))
+						{
+							NewFactoryNode->AddBooleanAttribute(StandardParameterName, Value);
+						}
+						break;
+					case UE::Interchange::EAttributeTypes::Float:
+						if (float Value; ShaderGraphNode.GetFloatAttribute(AttributeKey.Key, Value))
+						{
+							NewFactoryNode->AddFloatAttribute(StandardParameterName, Value);
+						}
+						break;
+					case UE::Interchange::EAttributeTypes::LinearColor:
+						if (FLinearColor Value; ShaderGraphNode.GetLinearColorAttribute(AttributeKey.Key, Value))
+						{
+							NewFactoryNode->AddLinearColorAttribute(StandardParameterName, Value);
+						}
+						break;
+					case UE::Interchange::EAttributeTypes::String:
+						if (FString TextureTranslatedNodeUid; ShaderGraphNode.GetStringAttribute(AttributeKey.Key, TextureTranslatedNodeUid))
+						{
+							// We only use string attributes for textures for now
+							FString FactoryTextureUid = UInterchangeFactoryBaseNode::BuildFactoryNodeUid(TextureTranslatedNodeUid);
+
+							NewFactoryNode->AddStringAttribute(StandardParameterName, FactoryTextureUid);
+							NewFactoryNode->AddFactoryDependencyUid(FactoryTextureUid);
+						}
+						break;
+					default:
+						break;
+				}
+			}
+			// It's another attribute related to primvar-compatible materials, just move it over to the factory node as-is
+			else if (AttributeKeyString.StartsWith(UE::Interchange::USD::ParameterToPrimvarAttributePrefix))
+			{
+				FString AttributeValue;
+				if (ShaderGraphNode.GetStringAttribute(AttributeKey.Key, AttributeValue))
+				{
+					NewFactoryNode->AddStringAttribute(AttributeKeyString, AttributeValue);
+				}
+			}
+			else if (AttributeKeyString.StartsWith(UE::Interchange::USD::PrimvarUVIndexAttributePrefix))
+			{
+				int32 AttributeValue;
+				if (ShaderGraphNode.GetInt32Attribute(AttributeKey.Key, AttributeValue))
+				{
+					NewFactoryNode->AddInt32Attribute(AttributeKeyString, AttributeValue);
+				}
+			}
+		}
+
+		return NewFactoryNode;
+	}
+
+	// Replaces dummy material instance factory nodes with new nodes where we properly fill in the
+	// material parameters and parent materials.
+	//
+	// We always translate the USD materials into UInterchangeShaderGraphNode (as opposed to directly
+	// UInterchangeMaterialInstanceNode) and fill in all parameters as base translated node attributes,
+	// in order to defer to the UInterchangeGenericMaterialPipeline (or other pipelines) whether to emit
+	// new material graphs or material instances.
+	//
+	// Here, we will likely run *after* UInterchangeGenericMaterialPipeline, which means is will already have
+	// converted our UInterchangeShaderGraphNodes into either material or material instance factory nodes,
+	// according to its "MaterialImport" option.
+	//
+	// If MaterialImport is set to "ImportAsMaterialInstances" however, UInterchangeGenericMaterialPipeline
+	// will not convert our base translated node attributes into the factory node like we need it to.
+	// This is what this function aims to repair: It will go over any of these broken material instance factory
+	// nodes, and replace them with new nodes that have the correct material info extracted from the
+	// corresponding ("target") UInterchangeShaderGraphNode base node attributes. It is also in charge of
+	// picking the correct parent material for these nodes, according to all the info added to the shader graph
+	// nodes.
+	//
+	// Note: This is all for the material instance code path. Unfortunately the full material graph code path
+	// won't work either, as we don't emit the shader graph translated node data from the USD translator just yet.
+	// We will likely want to support arbitrary shader graph generation for that though, which may come at a
+	// later point
+	void RepairMaterialInstanceFactoryNodes(UInterchangeBaseNodeContainer& NodeContainer)
+	{
+		TMap<FString, const UInterchangeShaderGraphNode*> MaterialFactoryNodeUidsToShaderGraphNodes;
+
+		auto CollectNodes =
+			[&NodeContainer, &MaterialFactoryNodeUidsToShaderGraphNodes](const FString& NodeUid, UInterchangeMaterialInstanceFactoryNode* FactoryNode)
+		{
+			TArray<FString> TargetNodeUids;
+			FactoryNode->GetTargetNodeUids(TargetNodeUids);
+
+			for (const FString& TargetNodeUid : TargetNodeUids)
+			{
+				const UInterchangeShaderGraphNode* ShaderGraphNode = Cast<UInterchangeShaderGraphNode>(NodeContainer.GetNode(TargetNodeUid));
+				if (!ShaderGraphNode)
+				{
+					continue;
+				}
+
+				// We explicitly flag our shader graph nodes with this attribute
+				bool bParse;
+				if (ShaderGraphNode->GetBooleanAttribute(*UE::Interchange::USD::ParseMaterialIdentifier, bParse) && bParse)
+				{
+					MaterialFactoryNodeUidsToShaderGraphNodes.Add(NodeUid, ShaderGraphNode);
+					break;
+				}
+			}
+		};
+		NodeContainer.IterateNodesOfType<UInterchangeMaterialInstanceFactoryNode>(CollectNodes);
+
+		for (const TPair<FString, const UInterchangeShaderGraphNode*>& Pair : MaterialFactoryNodeUidsToShaderGraphNodes)
+		{
+			const FString& FactoryNodeUid = Pair.Key;
+			const UInterchangeShaderGraphNode* ShaderGraphNode = Pair.Value;
+
+			UInterchangeMaterialInstanceFactoryNode* NewFactoryNode = BuildMaterialInstance(NodeContainer, *ShaderGraphNode);
+			if (NewFactoryNode)
+			{
+				NodeContainer.ReplaceNode(FactoryNodeUid, NewFactoryNode);
+			}
+		}
+	}
+
+	// Move custom attribute info from first volume nodes into corresponding factory nodes.
+	// This info originally came from the USD codeless SparseVolumeTextureAPI schema.
+	void MoveCustomAttributesToVolumeFactoryNode(const UInterchangeVolumeNode* VolumeNode, UInterchangeSparseVolumeTextureFactoryNode* FactoryNode)
+	{
+		using namespace UE::Interchange;
+		using namespace UE::Interchange::USD;
+
+		using StringSetterFunc = decltype(&UInterchangeSparseVolumeTextureFactoryNode::SetCustomAttributesAChannelX);
+		using EnumSetterFunc = decltype(&UInterchangeSparseVolumeTextureFactoryNode::SetCustomAttributesAFormat);
+
+		const static TMap<FAttributeKey, StringSetterFunc> StringSetters{
+			{FAttributeKey{SparseVolumeTexture::AttributesAChannelR}, &UInterchangeSparseVolumeTextureFactoryNode::SetCustomAttributesAChannelX},
+			{FAttributeKey{SparseVolumeTexture::AttributesAChannelG}, &UInterchangeSparseVolumeTextureFactoryNode::SetCustomAttributesAChannelY},
+			{FAttributeKey{SparseVolumeTexture::AttributesAChannelB}, &UInterchangeSparseVolumeTextureFactoryNode::SetCustomAttributesAChannelZ},
+			{FAttributeKey{SparseVolumeTexture::AttributesAChannelA}, &UInterchangeSparseVolumeTextureFactoryNode::SetCustomAttributesAChannelW},
+			{FAttributeKey{SparseVolumeTexture::AttributesBChannelR}, &UInterchangeSparseVolumeTextureFactoryNode::SetCustomAttributesBChannelX},
+			{FAttributeKey{SparseVolumeTexture::AttributesBChannelG}, &UInterchangeSparseVolumeTextureFactoryNode::SetCustomAttributesBChannelY},
+			{FAttributeKey{SparseVolumeTexture::AttributesBChannelB}, &UInterchangeSparseVolumeTextureFactoryNode::SetCustomAttributesBChannelZ},
+			{FAttributeKey{SparseVolumeTexture::AttributesBChannelA}, &UInterchangeSparseVolumeTextureFactoryNode::SetCustomAttributesBChannelW},
+		};
+
+		const static TMap<FAttributeKey, EnumSetterFunc> EnumSetters{
+			{FAttributeKey{SparseVolumeTexture::AttributesAFormat}, &UInterchangeSparseVolumeTextureFactoryNode::SetCustomAttributesAFormat},
+			{FAttributeKey{SparseVolumeTexture::AttributesBFormat}, &UInterchangeSparseVolumeTextureFactoryNode::SetCustomAttributesBFormat},
+		};
+
+		TArray<FAttributeKey> AttributeKeys;
+		VolumeNode->GetAttributeKeys(AttributeKeys);
+
+		for (const FAttributeKey& AttributeKey : AttributeKeys)
+		{
+			if (StringSetterFunc FoundStringSetter = StringSetters.FindRef(AttributeKey))
+			{
+				FString Value;
+				if (VolumeNode->GetStringAttribute(AttributeKey.ToString(), Value))
+				{
+					(FactoryNode->*FoundStringSetter)(Value);
+				}
+			}
+			else if (EnumSetterFunc FoundEnumSetter = EnumSetters.FindRef(AttributeKey))
+			{
+				int32 Value = 0;
+				if (VolumeNode->GetInt32Attribute(AttributeKey.ToString(), Value))
+				{
+					(FactoryNode->*FoundEnumSetter)(static_cast<EInterchangeSparseVolumeTextureFormat>(Value));
+				}
+			}
+		}
+	}
+
+	void ProcessHeterogeneousVolumeSceneNode(const UInterchangeSceneNode* SceneNode, UInterchangeBaseNodeContainer& NodeContainer)
+	{
+		// Find the corresponding actor factory node
+		UInterchangeHeterogeneousVolumeActorFactoryNode* ActorFactoryNode = nullptr;
+		{
+			TArray<FString> TargetNodeUids;
+			SceneNode->GetTargetNodeUids(TargetNodeUids);
+
+			for (const FString& TargetNodeUid : TargetNodeUids)
+			{
+				UInterchangeFactoryBaseNode* FactoryBaseNode = NodeContainer.GetFactoryNode(TargetNodeUid);
+				ActorFactoryNode = Cast<UInterchangeHeterogeneousVolumeActorFactoryNode>(FactoryBaseNode);
+				if (ActorFactoryNode)
+				{
+					break;
+				}
+			}
+		}
+
+		// Find the material it is using
+		UInterchangeMaterialInstanceFactoryNode* MaterialFactoryNode = nullptr;
+		{
+			TMap<FString, FString> OutMaterialDependencies;
+			SceneNode->GetSlotMaterialDependencies(OutMaterialDependencies);
+
+			if (FString* MaterialInstanceUid = OutMaterialDependencies.Find(UE::Interchange::Volume::VolumetricMaterial))
+			{
+				const FString MaterialFactoryNodeUid = UInterchangeFactoryBaseNode::BuildFactoryNodeUid(*MaterialInstanceUid);
+				MaterialFactoryNode = Cast<UInterchangeMaterialInstanceFactoryNode>(NodeContainer.GetFactoryNode(MaterialFactoryNodeUid));
+			}
+		}
+		if (!MaterialFactoryNode)
+		{
+			return;
+		}
+
+		// Make sure we spawn our actor after the material has finished generating, as we need to assign the
+		// material to the actor
+		if (ActorFactoryNode && MaterialFactoryNode)
+		{
+			ActorFactoryNode->AddFactoryDependencyUid(MaterialFactoryNode->GetUniqueID());
+		}
+
+		// Fixup the SVT assignment on the material, if needed
+		//
+		// Reference: CollectMaterialParameterTextureAssignment from USDVolVolumeTranslator.cpp
+		// We don't refactor/reuse the function as the logic/data is slightly different
+		//
+		// The USDTranslator has already translated into string attributes the best SVT to material parameter assignment
+		// we could come up with so far, which includes the fallback of interpreting the field names as material parameter names.
+		// In here, we check if it used that fallback, and if so we check if it will work or not, and if not we use another fallback.
+		// The intent here is to make sure that we get *something* assigned to the SVT material, like the legacy schema translator does
+		{
+			TMap<FString, FString> CleanAttributeNamesToVolumes;
+			bool bIsFallbackCase = false;
+
+			TArray<UE::Interchange::FAttributeKey> AllAttributeKeys;
+			MaterialFactoryNode->GetAttributeKeys(AllAttributeKeys);
+			for (const UE::Interchange::FAttributeKey& AttributeKey : AllAttributeKeys)
+			{
+				const UE::Interchange::EAttributeTypes AttributeType = MaterialFactoryNode->GetAttributeType(AttributeKey);
+				if (AttributeType != UE::Interchange::EAttributeTypes::String)
+				{
+					continue;
+				}
+
+				const FString AttributeKeyString = AttributeKey.ToString();
+				if (!AttributeKeyString.Contains(UE::Interchange::USD::VolumeFieldNameMaterialParameterPrefix))
+				{
+					continue;
+				}
+
+				bIsFallbackCase = true;
+
+				FString AttributeValue;
+				if (!MaterialFactoryNode->GetStringAttribute(AttributeKeyString, AttributeValue))
+				{
+					continue;
+				}
+
+				// e.g. go from "Inputs:USD_FieldName_density:Value" to "USD_FieldName_density"
+				FString CleanAttributeName = UInterchangeShaderPortsAPI::MakeInputName(AttributeKeyString);
+
+				// e.g. go "USD_FieldName_density" to "density"
+				CleanAttributeName.RemoveFromStart(UE::Interchange::USD::VolumeFieldNameMaterialParameterPrefix);
+
+				CleanAttributeNamesToVolumes.Add(CleanAttributeName, AttributeValue);
+
+				// We definitely don't want the fallback attrs to continue past this pipeline, as they have no meaning for Interchange itself
+				ensure(MaterialFactoryNode->RemoveAttribute(AttributeKeyString));
+			}
+			if (!bIsFallbackCase)
+			{
+				// If the translator didn't emit any of these VolumeFieldNameMaterialParameterPrefix attributes then we're not
+				// in a fallback material assignment case and can just stop now
+				return;
+			}
+
+			// Get the actual SVT parameter names from the material we're using on this actor
+			TArray<FString> SparseVolumeTextureParameterNames;
+			{
+				FString ParentMaterialContentPath;
+				if (!MaterialFactoryNode->GetCustomParent(ParentMaterialContentPath))
+				{
+					return;
+				}
+
+				FSoftObjectPath ReferencedObject{ParentMaterialContentPath};
+
+				// This is why this entire pipeline needs to be on the game thread
+				UMaterial* ParentMaterial = nullptr;
+				if (UMaterialInterface* MaterialInterface = Cast<UMaterialInterface>(ReferencedObject.TryLoad()))
+				{
+					ParentMaterial = MaterialInterface->GetMaterial();
+				}
+				if (!ParentMaterial)
+				{
+					return;
+				}
+
+				SparseVolumeTextureParameterNames = UsdUtils::GetSparseVolumeTextureParameterNames(ParentMaterial);
+			}
+
+			// Compensate for how parameter names are usually upper case in UE and field names are lower case in USD
+			TMap<FString, FString> LowercaseParamNamesToOriginal;
+			LowercaseParamNamesToOriginal.Reserve(SparseVolumeTextureParameterNames.Num());
+			for (const FString& ParamName : SparseVolumeTextureParameterNames)
+			{
+				LowercaseParamNamesToOriginal.Add(ParamName.ToLower(), ParamName);
+			}
+
+			// Check if our volume prim field names actually match material parameters (if we disregard casing)
+			bool bFoundFieldNameMatch = false;
+			for (const TPair<FString, FString>& AttrPair : CleanAttributeNamesToVolumes)
+			{
+				const FString& CleanAttributeName = AttrPair.Key;
+				const FString& VolumeUid = AttrPair.Value;
+
+				if (FString* FoundParamName = LowercaseParamNamesToOriginal.Find(CleanAttributeName.ToLower()))
+				{
+					bFoundFieldNameMatch = true;
+
+					const FString ParameterKey = UInterchangeShaderPortsAPI::MakeInputValueKey(*FoundParamName);
+					MaterialFactoryNode->AddStringAttribute(ParameterKey, VolumeUid);
+				}
+			}
+			if (bFoundFieldNameMatch)
+			{
+				// If we found any kind of match with the field names let's just take that
+				return;
+			}
+
+			// As a final fallback case just assign SVTs to the material parameters in alphabetical order like the legacy
+			// schema translator does
+
+			TArray<FString> VolumeUids;
+			{
+				TSet<FString> UniqueVolumeIds;
+				UniqueVolumeIds.Reserve(CleanAttributeNamesToVolumes.Num());
+				for (const TPair<FString, FString>& ParamPair : CleanAttributeNamesToVolumes)
+				{
+					UniqueVolumeIds.Add(ParamPair.Value);
+				}
+				VolumeUids = UniqueVolumeIds.Array();
+			}
+
+			SparseVolumeTextureParameterNames.Sort();
+			VolumeUids.Sort();
+
+			for (int32 Index = 0; Index < VolumeUids.Num() && Index < SparseVolumeTextureParameterNames.Num(); ++Index)
+			{
+				const FString& VolumeUid = VolumeUids[Index];
+				const FString& ParameterName = SparseVolumeTextureParameterNames[Index];
+				const FString ParameterKey = UInterchangeShaderPortsAPI::MakeInputValueKey(ParameterName);
+
+				MaterialFactoryNode->AddStringAttribute(ParameterKey, VolumeUid);
+			}
+		}
+	}
+
+	void ProcessVolumeNodes(UInterchangeBaseNodeContainer& NodeContainer)
+	{
+		using namespace UE::Interchange;
+		using namespace UE::Interchange::USD;
+
+		TMap<const UInterchangeSceneNode*, const UInterchangeVolumeNode*> SceneNodeToFirstVolume;
+		TMap<const UInterchangeVolumeNode*, UInterchangeSparseVolumeTextureFactoryNode*> FirstVolumeToFactoryNode;
+
+		// Collect all the node pairings we need
+		NodeContainer.IterateNodesOfType<UInterchangeSparseVolumeTextureFactoryNode>(
+			[&SceneNodeToFirstVolume,
+			 &NodeContainer,
+			 &FirstVolumeToFactoryNode](const FString& FactoryNodeUid, UInterchangeSparseVolumeTextureFactoryNode* FactoryNode)
+			{
+				FString FirstVolumeUid = UInterchangeFactoryBaseNode::BuildTranslatedNodeUid(FactoryNodeUid);
+
+				const UInterchangeVolumeNode* FirstVolume = Cast<const UInterchangeVolumeNode>(NodeContainer.GetNode(FirstVolumeUid));
+				ensure(FirstVolume);
+
+				FirstVolumeToFactoryNode.Add(FirstVolume, FactoryNode);
+
+				// Find our scene node too
+				TArray<FString> TargetNodeUids;
+				FirstVolume->GetTargetNodeUids(TargetNodeUids);
+				for (const FString& TargetNodeUid : TargetNodeUids)
+				{
+					const UInterchangeBaseNode* BaseNode = NodeContainer.GetNode(TargetNodeUid);
+					if (const UInterchangeSceneNode* SceneNode = Cast<const UInterchangeSceneNode>(BaseNode))
+					{
+						SceneNodeToFirstVolume.Add(SceneNode, FirstVolume);
+						break;
+					}
+				}
+			}
+		);
+
+		for (TPair<const UInterchangeVolumeNode*, UInterchangeSparseVolumeTextureFactoryNode*> Pair : FirstVolumeToFactoryNode)
+		{
+			const UInterchangeVolumeNode* VolumeNode = Pair.Key;
+			UInterchangeSparseVolumeTextureFactoryNode* FactoryNode = Pair.Value;
+			MoveCustomAttributesToVolumeFactoryNode(VolumeNode, FactoryNode);
+		}
+
+		for (const TPair<const UInterchangeSceneNode*, const UInterchangeVolumeNode*>& Pair : SceneNodeToFirstVolume)
+		{
+			const UInterchangeSceneNode* SceneNode = Pair.Key;
+			ProcessHeterogeneousVolumeSceneNode(SceneNode, NodeContainer);
+		}
+	}
+
+	void RemovePseudoRootFactoryNode(UInterchangeBaseNodeContainer& NodeContainer)
+	{
+		const FString& FactoryNodeUid = GetPseudoRootFactoryNodeUid();
+		UInterchangeFactoryBaseNode* PseudoRootFactoryNode = NodeContainer.GetFactoryNode(FactoryNodeUid);
+		if (PseudoRootFactoryNode)
+		{
+			// Move all root node children to top level
+			for (const FString& ChildNodeUid : NodeContainer.GetNodeChildrenUids(FactoryNodeUid))
+			{
+				if (UInterchangeFactoryBaseNode* ChildNode = NodeContainer.GetFactoryNode(ChildNodeUid))
+				{
+					NodeContainer.ClearNodeParentUid(ChildNodeUid);
+				}
+			}
+
+			// Disable the pseudoroot itself (we shouldn't have any asset node for it though)
+			NodeContainer.ReplaceNode(FactoryNodeUid, nullptr);
+		}
+	}
+
+#if WITH_EDITORONLY_DATA
+	/**
+	 * Add the primvar attributes from the mesh nodes to factory nodes
+	 */
+	void AddPrimvarAttributesToFactoryNodes(
+		const UInterchangeMeshNode* MeshNode,
+		UInterchangeMeshFactoryNode* MeshFactoryNode,
+		UInterchangeBaseNodeContainer* NodeContainer
+	)
+	{
+		using namespace UE::Interchange;
+
+		if (int32 NumPrimvar; MeshNode->GetInt32Attribute(USD::Primvar::Number, NumPrimvar))
+		{
+			MeshFactoryNode->AddInt32Attribute(USD::Primvar::Number, NumPrimvar);
+			for (int32 Index = 0; Index < NumPrimvar; ++Index)
+			{
+				FString PrimvarAttribute = USD::Primvar::Name + FString::FromInt(Index);
+				if (FString PrimvarName; MeshNode->GetStringAttribute(PrimvarAttribute, PrimvarName))
+				{
+					MeshFactoryNode->AddStringAttribute(PrimvarAttribute, PrimvarName);
+				}
+
+				if (bool bTangentSpace; MeshNode->GetBooleanAttribute(PrimvarAttribute, bTangentSpace))
+				{
+					MeshFactoryNode->AddBooleanAttribute(PrimvarAttribute, bTangentSpace);
+				}
+
+				FString ShaderNodeTextureSampleAttribute = USD::Primvar::ShaderNodeTextureSample + FString::FromInt(Index);
+				if (FString ShaderNodeTextureSampleUID; MeshNode->GetStringAttribute(ShaderNodeTextureSampleAttribute, ShaderNodeTextureSampleUID))
+				{
+					// At this stage the USD pipeline should be the latest one in the stack, we should retrieve the MaterialExpression
+					// Factory Node associated to this ShaderNode
+					FString TextureSampleFactoryNodeUID = UInterchangeFactoryBaseNode::BuildFactoryNodeUid(ShaderNodeTextureSampleUID);
+
+					if (Cast<UInterchangeMaterialExpressionFactoryNode>(NodeContainer->GetFactoryNode(TextureSampleFactoryNodeUID)))
+					{
+						MeshFactoryNode->AddStringAttribute(ShaderNodeTextureSampleAttribute, TextureSampleFactoryNodeUID);
+					}
+				}
+			}
+		}
+	}
+#endif	  // WITH_EDITORONLY_DATA
+
+	struct FMaterialPrimvarInfo
+	{
+		TMap<FString, FString> ParameterToPrimvar;
+		TMap<FString, int32> PrimvarToUVIndex;
+	};
+
+	// Gets or creates a filled out FMaterialPrimvarInfo struct for a particular factory node.
+	// We use this because we may iterate over the same material many times during ProcessMeshNodes,
+	// and we don't want to recompute this info every time
+	FMaterialPrimvarInfo* GetOrCreateMaterialPrimvarInfo(
+		UInterchangeMaterialInstanceFactoryNode* InstanceFactoryNode,
+		TMap<FString, FMaterialPrimvarInfo>& InOutFactoryNodeUidToInfo
+	)
+	{
+		using namespace UE::Interchange;
+
+		if (!InstanceFactoryNode)
+		{
+			return nullptr;
+		}
+
+		const FString& MaterialUid = InstanceFactoryNode->GetUniqueID();
+
+		FMaterialPrimvarInfo* FoundInfo = InOutFactoryNodeUidToInfo.Find(MaterialUid);
+		if (!FoundInfo)
+		{
+			FMaterialPrimvarInfo& NewInfo = InOutFactoryNodeUidToInfo.Emplace(MaterialUid);
+
+			TArray<UE::Interchange::FAttributeKey> MaterialAttributeKeys;
+			InstanceFactoryNode->GetAttributeKeys(MaterialAttributeKeys);
+
+			NewInfo.ParameterToPrimvar.Reserve(MaterialAttributeKeys.Num());
+			NewInfo.PrimvarToUVIndex.Reserve(MaterialAttributeKeys.Num());
+			for (const UE::Interchange::FAttributeKey& AttributeKey : MaterialAttributeKeys)
+			{
+				const FString& AttributeKeyString = AttributeKey.ToString();
+
+				if (AttributeKeyString.StartsWith(USD::ParameterToPrimvarAttributePrefix))
+				{
+					FString ParameterName = AttributeKeyString;
+					ensure(ParameterName.RemoveFromStart(USD::ParameterToPrimvarAttributePrefix));
+
+					FString Primvar;
+					if (InstanceFactoryNode->GetStringAttribute(AttributeKeyString, Primvar))
+					{
+						NewInfo.ParameterToPrimvar.Add(ParameterName, Primvar);
+					}
+				}
+				else if (AttributeKeyString.StartsWith(USD::PrimvarUVIndexAttributePrefix))
+				{
+					FString ParameterName = AttributeKeyString;
+					ensure(ParameterName.RemoveFromStart(USD::PrimvarUVIndexAttributePrefix));
+
+					int32 UVIndex = -1;
+					if (InstanceFactoryNode->GetInt32Attribute(AttributeKeyString, UVIndex))
+					{
+						NewInfo.PrimvarToUVIndex.Add(ParameterName, UVIndex);
+					}
+				}
+			}
+
+			FoundInfo = &NewInfo;
+		}
+
+		return FoundInfo;
+	}
+
+	// Returns a primvar-compatible version of OriginalFactoryNode.
+	// i.e. creates (if needed) a new MaterialInstanceFactoryNode that has a primvar to UV index
+	// mapping that is usable with the primvar to UV index layout in MeshPrimvarToUVIndex
+	UInterchangeMaterialInstanceFactoryNode* GetOrCreateCompatibleMaterial(
+		const UInterchangeMaterialInstanceFactoryNode* OriginalFactoryNode,
+		const FMaterialPrimvarInfo* OriginalMaterialPrimvarInfo,
+		const TMap<FString, int32>& MeshPrimvarToUVIndex,
+		UInterchangeBaseNodeContainer* InNodeContainer
+	)
+	{
+		using namespace UE::Interchange;
+
+		if (!OriginalFactoryNode || !OriginalMaterialPrimvarInfo || !InNodeContainer)
+		{
+			return nullptr;
+		}
+
+		// First, let's create the target primvar UVIndex assignment that is compatible with this mesh.
+		// We use an array of TPairs here so that we can sort these into a deterministic order for hashing later.
+		TArray<TPair<FString, int32>> CompatiblePrimvarAndUVIndexPairs;
+		CompatiblePrimvarAndUVIndexPairs.Reserve(OriginalMaterialPrimvarInfo->PrimvarToUVIndex.Num());
+		for (const TPair<FString, int32>& Pair : OriginalMaterialPrimvarInfo->PrimvarToUVIndex)
+		{
+			const FString& MaterialPrimvar = Pair.Key;
+
+			bool bFoundUVIndex = false;
+
+			// Mesh has this primvar available at some UV index, point to it
+			if (const int32* FoundMeshUVIndex = MeshPrimvarToUVIndex.Find(MaterialPrimvar))
+			{
+				int32 MeshUVIndex = *FoundMeshUVIndex;
+				if (MeshUVIndex >= 0 && MeshUVIndex < USD_PREVIEW_SURFACE_MAX_UV_SETS)
+				{
+					CompatiblePrimvarAndUVIndexPairs.Add(TPair<FString, int32>{MaterialPrimvar, MeshUVIndex});
+					bFoundUVIndex = true;
+				}
+			}
+
+			if (!bFoundUVIndex)
+			{
+				// Point this primvar to read an unused UV index instead, since our mesh doesn't have this primvar
+				CompatiblePrimvarAndUVIndexPairs.Add(TPair<FString, int32>{MaterialPrimvar, UNUSED_UV_INDEX});
+			}
+		}
+
+		// Generate a deterministic hash based on the original material hash and this primvar UVIndex assignment
+		CompatiblePrimvarAndUVIndexPairs.Sort(
+			[](const TPair<FString, int32>& LHS, const TPair<FString, int32>& RHS)
+			{
+				if (LHS.Key == RHS.Key)
+				{
+					return LHS.Value < RHS.Value;
+				}
+				else
+				{
+					return LHS.Key < RHS.Key;
+				}
+			}
+		);
+		FSHAHash Hash;
+		FSHA1 SHA1;
+		const FString OriginalNodeUid = OriginalFactoryNode->GetUniqueID();
+		SHA1.UpdateWithString(*OriginalNodeUid, OriginalNodeUid.Len());
+		for (const TPair<FString, int32>& Pair : CompatiblePrimvarAndUVIndexPairs)
+		{
+			SHA1.UpdateWithString(*Pair.Key, Pair.Key.Len());
+			SHA1.Update((const uint8*)&Pair.Value, sizeof(Pair.Value));
+		}
+		SHA1.Final();
+		SHA1.GetHash(&Hash.Hash[0]);
+
+		const FString CompatibleNodeUid = OriginalNodeUid + UE::Interchange::USD::CompatibleMaterialUidSuffix + Hash.ToString();
+
+		// Check if we made this compatible material before
+		UInterchangeMaterialInstanceFactoryNode* CompatibleFactoryNode = Cast<UInterchangeMaterialInstanceFactoryNode>(
+			InNodeContainer->GetFactoryNode(CompatibleNodeUid)
+		);
+		if (CompatibleFactoryNode)
+		{
+			return CompatibleFactoryNode;
+		}
+
+		// We need to actually create a new compatible material
+		CompatibleFactoryNode = NewObject<UInterchangeMaterialInstanceFactoryNode>(InNodeContainer);
+
+		UObject* Object = nullptr;
+		CompatibleFactoryNode->CopyWithObject(OriginalFactoryNode, Object);
+
+		CompatibleFactoryNode->InitializeNode(
+			CompatibleNodeUid,
+			OriginalFactoryNode->GetDisplayLabel() + TEXT("_Compatible"),
+			OriginalFactoryNode->GetNodeContainerType()
+		);
+
+		InNodeContainer->AddNode(CompatibleFactoryNode);
+
+		// Cleanup our old primvar to UV Index mapping in case we don't overlap perfectly
+		for (const TPair<FString, int32>& OldPair : OriginalMaterialPrimvarInfo->PrimvarToUVIndex)
+		{
+			const FString& Primvar = OldPair.Key;
+			const FString AttributeKeyString = USD::PrimvarUVIndexAttributePrefix + Primvar;
+
+			CompatibleFactoryNode->RemoveAttribute(AttributeKeyString);
+		}
+
+		// We also need to move our pairs into an actual TMap now, as we'll need to query into it later
+		TMap<FString, int32> CompatiblePrimvarToUVIndex;
+		CompatiblePrimvarToUVIndex.Reserve(CompatiblePrimvarAndUVIndexPairs.Num());
+
+		// Set our new primvar to UV index mapping on the factory node (only for the sake of information / debugging)
+		for (const TPair<FString, int32>& NewPair : CompatiblePrimvarAndUVIndexPairs)
+		{
+			const FString& Primvar = NewPair.Key;
+			const int32 UVIndex = NewPair.Value;
+			const FString AttributeKeyString = USD::PrimvarUVIndexAttributePrefix + Primvar;
+
+			CompatibleFactoryNode->AddInt32Attribute(AttributeKeyString, UVIndex);
+
+			CompatiblePrimvarToUVIndex.Add(Primvar, UVIndex);
+		}
+
+		// Set our compatible UV indices into the actual material parameter attributes that will be set
+		// on the material instance
+		for (const TPair<FString, FString>& ParameterPair : OriginalMaterialPrimvarInfo->ParameterToPrimvar)
+		{
+			const FString& MaterialParameter = ParameterPair.Key;
+			const FString& Primvar = ParameterPair.Value;
+
+			int32 UVIndex = UNUSED_UV_INDEX;
+			if (int32* FoundUVIndex = CompatiblePrimvarToUVIndex.Find(Primvar))
+			{
+				UVIndex = *FoundUVIndex;
+			}
+
+			const bool bEnableTexture = UVIndex >= 0 && UVIndex < UNUSED_UV_INDEX;
+			const FString EnableKey = UInterchangeShaderPortsAPI::MakeInputValueKey(
+				USD::UseTextureParameterPrefix + MaterialParameter + USD::UseTextureParameterSuffix
+			);
+			ensure(CompatibleFactoryNode->AddFloatAttribute(EnableKey, bEnableTexture));
+
+			const FString UVIndexKey = UInterchangeShaderPortsAPI::MakeInputValueKey(MaterialParameter + USD::UVIndexParameterSuffix);
+			if (bEnableTexture)
+			{
+				ensure(CompatibleFactoryNode->AddFloatAttribute(UVIndexKey, UVIndex));
+			}
+			else
+			{
+				ensure(CompatibleFactoryNode->RemoveAttribute(UVIndexKey));
+			}
+		}
+
+		return CompatibleFactoryNode;
+	}
+
+	// Makes sure that all materials in the MeshNode's SlotDependencies array are primvar-compatible
+	// with the primvar to UV index mapping of the mesh, creating new materials and reassigning them to the
+	// SlotDependencies custom attribute if needed.
+	void GeneratePrimvarCompatibleMaterials(
+		const UInterchangeMeshNode* InMeshNode,
+		UInterchangeMeshFactoryNode* InMeshFactoryNode,
+		UInterchangeBaseNodeContainer* InNodeContainer,
+		TMap<FString, FMaterialPrimvarInfo>& InOutFactoryNodeUidToInfo
+	)
+	{
+		using namespace UE::Interchange;
+
+		// Rebuild the mesh's primvar to UV index map
+		TMap<FString, int32> MeshPrimvarToUVIndex;
+		{
+			TArray<UE::Interchange::FAttributeKey> MeshAttributeKeys;
+			InMeshNode->GetAttributeKeys(MeshAttributeKeys);
+
+			MeshPrimvarToUVIndex.Reserve(MeshAttributeKeys.Num());
+
+			for (const UE::Interchange::FAttributeKey& AttributeKey : MeshAttributeKeys)
+			{
+				const FString& AttributeKeyString = AttributeKey.ToString();
+
+				if (AttributeKeyString.StartsWith(USD::PrimvarUVIndexAttributePrefix))
+				{
+					FString ParameterName = AttributeKeyString;
+					ensure(ParameterName.RemoveFromStart(USD::PrimvarUVIndexAttributePrefix));
+
+					int32 UVIndex = -1;
+					if (InMeshNode->GetInt32Attribute(AttributeKeyString, UVIndex))
+					{
+						MeshPrimvarToUVIndex.Add(ParameterName, UVIndex);
+					}
+				}
+			}
+		}
+
+		// Build the reverse map too. This because we can still consider a mesh and material "compatible"
+		// if the material tries reading "st2" and the mesh doesn't have "st2" anywhere. We will just turn off
+		// usage of that particular texture on the material instance.
+		//
+		// For that case to still count as "compatible" though, then whatever UV index the material is trying to
+		// read "st2" from must not be used by any other primvar: Otherwise the mesh will still be providing
+		// some primvar data to the material, but the material will treat it as if it were "st2" when it is not
+		TArray<TSet<FString>> UVIndexToMeshPrimvars;
+		UVIndexToMeshPrimvars.SetNum(USD_PREVIEW_SURFACE_MAX_UV_SETS);
+		for (const TPair<FString, int32>& MeshPair : MeshPrimvarToUVIndex)
+		{
+			if (UVIndexToMeshPrimvars.IsValidIndex(MeshPair.Value))
+			{
+				UVIndexToMeshPrimvars[MeshPair.Value].Add(MeshPair.Key);
+			}
+		}
+
+		// Use primvar to UV index information to construct primvar-compatible materials if needed
+		TMap<FString, FString> SlotNameToMaterialFactoryNodeUid;
+		InMeshFactoryNode->GetSlotMaterialDependencies(SlotNameToMaterialFactoryNodeUid);
+		for (const TPair<FString, FString>& SlotPair : SlotNameToMaterialFactoryNodeUid)
+		{
+			const FString& SlotName = SlotPair.Key;
+			const FString& MaterialFactoryNodeUid = SlotPair.Value;
+
+			UInterchangeMaterialInstanceFactoryNode* OriginalMaterial = Cast<UInterchangeMaterialInstanceFactoryNode>(
+				InNodeContainer->GetFactoryNode(MaterialFactoryNodeUid)
+			);
+			if (!OriginalMaterial)
+			{
+				continue;
+			}
+
+			FMaterialPrimvarInfo* OriginalMaterialInfo = GetOrCreateMaterialPrimvarInfo(OriginalMaterial, InOutFactoryNodeUidToInfo);
+			if (!OriginalMaterialInfo)
+			{
+				continue;
+			}
+
+			// Check if the material's primvar-UVIndex mapping matches the mesh
+			bool bCompatible = true;
+			for (const TPair<FString, int32>& Pair : OriginalMaterialInfo->PrimvarToUVIndex)
+			{
+				const FString& MaterialPrimvar = Pair.Key;
+				int32 MaterialUVIndex = Pair.Value;
+
+				// If the mesh has the same primvar the material wants, it should be at the same UVIndex the material
+				// will read from
+				if (const int32* MeshUVIndex = MeshPrimvarToUVIndex.Find(MaterialPrimvar))
+				{
+					if (*MeshUVIndex != MaterialUVIndex)
+					{
+						bCompatible = false;
+						break;
+					}
+				}
+
+				// If the material is going to read from a given UVIndex that exists on the mesh, that UV set should
+				// contain the primvar data that the material expects to read
+				if (UVIndexToMeshPrimvars.IsValidIndex(MaterialUVIndex))
+				{
+					const TSet<FString>& CompatiblePrimvars = UVIndexToMeshPrimvars[MaterialUVIndex];
+					if (!CompatiblePrimvars.Contains(MaterialPrimvar))
+					{
+						bCompatible = false;
+						break;
+					}
+				}
+			}
+
+			if (!bCompatible)
+			{
+				// Generate a primvar-compatible material
+				UInterchangeMaterialInstanceFactoryNode* CompatibleMaterial = GetOrCreateCompatibleMaterial(	//
+					OriginalMaterial,
+					OriginalMaterialInfo,
+					MeshPrimvarToUVIndex,
+					InNodeContainer
+				);
+				if (CompatibleMaterial)
+				{
+					InMeshFactoryNode->SetSlotMaterialDependencyUid(SlotName, CompatibleMaterial->GetUniqueID());
+				}
+			}
+		}
+	}
+
+	void ProcessMeshNodes(UInterchangeBaseNodeContainer* NodeContainer, bool bGenerateCompatibleMaterials)
+	{
+		using namespace UE::Interchange;
+
+		// Collect the nodes in a separate container because GeneratePrimvarCompatibleMaterials may add additional nodes
+		// to the node container itself, which we can't do while we iterate over it
+		TMap<const UInterchangeMeshNode*, UInterchangeMeshFactoryNode*> MeshNodeToFactoryNode;
+
+		NodeContainer->IterateNodesOfType<UInterchangeMeshNode>(
+			[NodeContainer, &MeshNodeToFactoryNode](const FString& NodeUid, UInterchangeMeshNode* MeshNode)
+			{
+				const FString FactoryNodeUid = UInterchangeFactoryBaseNode::BuildFactoryNodeUid(NodeUid);
+				UInterchangeMeshFactoryNode* MeshFactoryNode = Cast<UInterchangeMeshFactoryNode>(NodeContainer->GetFactoryNode(FactoryNodeUid));
+				if (!MeshFactoryNode)
+				{
+					return;
+				}
+
+				MeshNodeToFactoryNode.Add(MeshNode, MeshFactoryNode);
+
+#if WITH_EDITORONLY_DATA
+				AddPrimvarAttributesToFactoryNodes(MeshNode, MeshFactoryNode, NodeContainer);
+#endif	  // WITH_EDITORONLY_DATA
+			}
+		);
+
+		if (bGenerateCompatibleMaterials)
+		{
+			// Build this map to prevent us from reparsing the same material in case multiple meshes
+			// are using them
+			TMap<FString, FMaterialPrimvarInfo> MaterialFactoryNodeUidToInfo;
+			for (const TPair<const UInterchangeMeshNode*, UInterchangeMeshFactoryNode*>& MeshPair : MeshNodeToFactoryNode)
+			{
+				GeneratePrimvarCompatibleMaterials(MeshPair.Key, MeshPair.Value, NodeContainer, MaterialFactoryNodeUidToInfo);
+			}
+		}
+	}
+
+#if WITH_EDITORONLY_DATA
+	class FUsdPrimvarBaker
+	{
+	public:
+
+		FUsdPrimvarBaker(
+			UStaticMesh* StaticMeshAsset,
+			const UInterchangeBaseNodeContainer* NodeContainer,
+			const UInterchangeFactoryBaseNode* FactoryNode
+		)
+			: NodeContainer{NodeContainer}
+			, FactoryNode{FactoryNode}
+			, StaticMeshAsset{StaticMeshAsset}
+		{
+		}
+
+		void Bake()
+		{
+			using namespace UE::Interchange;
+			using namespace UE::Interchange::Materials;
+			using namespace UE::Geometry;
+
+			if (int32 NumberOfPrimvar; FactoryNode->GetInt32Attribute(USD::Primvar::Number, NumberOfPrimvar) && NumberOfPrimvar)
+			{
+				FMeshDescription MeshDescriptionPrimvar = CreateMeshDescription();
+
+				FStaticMeshAttributes MeshAttributes(MeshDescriptionPrimvar);
+				TVertexInstanceAttributesRef<FVector4f> VertexColors = MeshAttributes.GetVertexInstanceColors();
+				TVertexInstanceAttributesRef<FVector3f> Normals = MeshAttributes.GetVertexInstanceNormals();
+
+				for (int32 Index = 0; Index < NumberOfPrimvar; ++Index)
+				{
+					UMaterialExpressionTextureSample* TextureSampleExpression = GetMaterialExpressionTextureSample(Index);
+
+					if (!TextureSampleExpression)
+					{
+						continue;
+					}
+
+					FString PrimvarName;
+					bool bTangentSpace = false;
+					FactoryNode->GetStringAttribute(USD::Primvar::Name + FString::FromInt(Index), PrimvarName);
+					FactoryNode->GetBooleanAttribute(USD::Primvar::TangentSpace + FString::FromInt(Index), bTangentSpace);
+
+					TVertexInstanceAttributesRef<FVector4f> Primvar = MeshDescriptionPrimvar.VertexInstanceAttributes().GetAttributesRef<FVector4f>(
+						*PrimvarName
+					);
+
+					if (!(Primvar.IsValid() && Primvar.GetNumElements() == VertexColors.GetNumElements()
+						  && Primvar.GetNumElements() == Normals.GetNumElements()))
+					{
+						continue;
+					}
+
+					// We put the primvar in the Color/Normal channel, that way during the conversion to a DynamicMesh it will also handle the
+					// non-manifold case
+					for (const int32 VertexInstanceID : MeshDescriptionPrimvar.VertexInstances().GetElementIDs())
+					{
+						bTangentSpace ? Normals.Set(VertexInstanceID, Primvar.Get(VertexInstanceID))
+									  : VertexColors.Set(VertexInstanceID, Primvar.Get(VertexInstanceID));
+					}
+
+					FDynamicMesh3 DynamicMesh;
+					FMeshDescriptionToDynamicMesh Convert;
+					Convert.Convert(&MeshDescriptionPrimvar, DynamicMesh);
+
+					const FDynamicMeshAABBTree3 DetailSpatial(&DynamicMesh);
+					FMeshBakerDynamicMeshSampler DetailSampler(&DynamicMesh, &DetailSpatial);
+					TSharedPtr<FMeshPropertyMapEvaluator> PropertyEval = MakeShared<FMeshPropertyMapEvaluator>();
+					PropertyEval->Property = bTangentSpace ? EMeshPropertyMapType::Normal : EMeshPropertyMapType::VertexColor;
+
+					constexpr int32 TextureSize = 1024;
+					const FImageDimensions ImageDimensions(TextureSize, TextureSize);
+					FMeshMapBaker Baker;
+					Baker.SetCorrespondenceStrategy(FMeshBaseBaker::ECorrespondenceStrategy::Identity);
+					Baker.SetTargetMesh(&DynamicMesh);
+					Baker.SetDetailSampler(&DetailSampler);
+					Baker.AddEvaluator(PropertyEval);
+					Baker.SetTargetMeshUVLayer(UVChannel);
+					Baker.SetDimensions(ImageDimensions);
+					Baker.SetProjectionDistance(3.0f);
+					Baker.SetSamplesPerPixel(1);
+					Baker.SetFilter(FMeshMapBaker::EBakeFilterType::BSpline);
+					Baker.SetGutterEnabled(true);
+					Baker.SetGutterSize(4);
+					Baker.SetTileSize(TextureSize);
+					Baker.Bake();
+
+					FTexture2DBuilder TextureBuilder;
+					const FTexture2DBuilder::ETextureType TextureType = GetTextureType(TextureSampleExpression);
+
+					TextureBuilder
+						.InitializeAndReplaceExistingTexture(Cast<UTexture2D>(TextureSampleExpression->Texture.Get()), TextureType, ImageDimensions);
+					TextureBuilder.Copy(*Baker.GetBakeResults(0)[0]);
+					TextureBuilder.Commit();
+
+					TextureSampleExpression->ConstCoordinate = UVChannel;
+
+					MaterialExpressionTextureSamples.Add(TextureSampleExpression);
+				}
+
+				UpdateMaterials();
+			}
+		}
+
+	private:
+
+		UMaterialExpressionTextureSample* GetMaterialExpressionTextureSample(int32 IndexPrimvar) const
+		{
+			using namespace UE::Interchange;
+			using namespace UE::Interchange::Materials;
+
+			UMaterialExpressionTextureSample* MaterialExpressionTextureSample = nullptr;
+
+			FString ShaderNodeTextureSampleUID;
+			FactoryNode->GetStringAttribute(USD::Primvar::ShaderNodeTextureSample + FString::FromInt(IndexPrimvar), ShaderNodeTextureSampleUID);
+			if (UInterchangeMaterialExpressionFactoryNode* TextureSampleExpressionFactoryNode = Cast<UInterchangeMaterialExpressionFactoryNode>(
+					NodeContainer->GetFactoryNode(ShaderNodeTextureSampleUID)
+				))
+			{
+				if (FString ExpressionPath;
+					TextureSampleExpressionFactoryNode->GetStringAttribute(Factory::Expression::Path.ToString(), ExpressionPath))
+				{
+					MaterialExpressionTextureSample = Cast<UMaterialExpressionTextureSample>(FSoftObjectPath{ExpressionPath}.TryLoad());
+				}
+			}
+
+			return MaterialExpressionTextureSample;
+		}
+
+		/** Create a mesh description based on the one from the StaticMesh, the StaticMesh will now hold a new UV set */
+		FMeshDescription CreateMeshDescription()
+		{
+			FMeshDescription* MeshDescription = StaticMeshAsset->GetMeshDescription(0);
+			FStaticMeshAttributes StaticMeshAttributes(*MeshDescription);
+			TVertexInstanceAttributesRef<FVector2f> VertexInstanceUVs = StaticMeshAttributes.GetVertexInstanceUVs();
+
+			if (AreUVsOverlapping(VertexInstanceUVs))
+			{
+				FStaticMeshOperations::FGenerateUVOptions Options;
+				Options.UVMethod = FStaticMeshOperations::EGenerateUVMethod::Default;
+				TArray<FVector2D> UVs;
+				FStaticMeshOperations::GenerateUV(*MeshDescription, Options, UVs);
+				FStaticMeshOperations::AddUVChannel(*MeshDescription);
+				UVChannel = VertexInstanceUVs.GetNumChannels() - 1;
+
+				for (const int32 VertexInstanceID : MeshDescription->VertexInstances().GetElementIDs())
+				{
+					VertexInstanceUVs.Set(VertexInstanceID, UVChannel, FVector2f{UVs[VertexInstanceID]});
+				}
+			}
+
+			FMeshDescription MeshDescriptionPrimvar(*MeshDescription);
+			return MeshDescriptionPrimvar;
+		}
+
+		void UpdateMaterials() const
+		{
+			TSet<UMaterial*> Materials;
+			for (UMaterialExpressionTextureSample* TextureSampleExpression : MaterialExpressionTextureSamples)
+			{
+				Materials.FindOrAdd(TextureSampleExpression->Material);
+				TextureSampleExpression->PostEditChange();
+				TextureSampleExpression->Modify();
+			}
+
+			for (UMaterial* Material : Materials)
+			{
+				Material->PostEditChange();
+				Material->ForceRecompileForRendering();
+				Material->MarkPackageDirty();
+			}
+		}
+
+		UE::Geometry::FTexture2DBuilder::ETextureType GetTextureType(UMaterialExpressionTextureSample* TextureSampleExpression)
+		{
+			using namespace UE::Geometry;
+			switch (TextureSampleExpression->SamplerType)
+			{
+				case EMaterialSamplerType::SAMPLERTYPE_Color:
+					return FTexture2DBuilder::ETextureType::Color;
+
+				case EMaterialSamplerType::SAMPLERTYPE_Normal:
+					return FTexture2DBuilder::ETextureType::NormalMap;
+
+				default:
+					return FTexture2DBuilder::ETextureType::ColorLinear;
+			}
+		}
+
+		struct FUVKeyFunc : public BaseKeyFuncs<FVector2f, FVector2f>
+		{
+			static constexpr float Threshold = 0.00001;
+
+			static FORCEINLINE KeyInitType GetSetKey(ElementInitType Element)
+			{
+				return Element;
+			}
+
+			static FORCEINLINE uint32 GetKeyHash(ElementInitType& Key)
+			{
+				int32 RoundedU = FMath::RoundToInt(Key.X / Threshold);
+				int32 RoundedV = FMath::RoundToInt(Key.Y / Threshold);
+				return HashCombine(GetTypeHash(RoundedU), GetTypeHash(RoundedV));
+			}
+
+			static FORCEINLINE bool Matches(ElementInitType A, ElementInitType B)
+			{
+				return FMath::IsNearlyEqual(A.X, B.X, Threshold) && FMath::IsNearlyEqual(A.Y, B.Y, Threshold);
+			}
+		};
+
+		bool AreUVsOverlapping(const TVertexInstanceAttributesRef<FVector2f>& UVs)
+		{
+			FMeshDescription* MeshDecription = StaticMeshAsset->GetMeshDescription(0);
+			TSet<FVector2f, FUVKeyFunc> UVSet;
+
+			for (const FVertexInstanceID VertexInstanceID : MeshDecription->VertexInstances().GetElementIDs())
+			{
+				FVector2f UV = UVs.Get(VertexInstanceID);
+
+				if (UVSet.Contains(UV))
+				{
+					UE_LOG(LogInterchangeUSDPipeline, Warning, TEXT("UVs are Overlapping"));
+					return true;
+				}
+
+				UVSet.Add(UV);
+			}
+
+			return false;
+		}
+
+	private:
+
+		const UInterchangeBaseNodeContainer* NodeContainer;
+		const UInterchangeFactoryBaseNode* FactoryNode;
+		UStaticMesh* StaticMeshAsset;
+		TArray<UMaterialExpressionTextureSample*> MaterialExpressionTextureSamples;
+		int32 UVChannel = 0;
+	};
+#endif	  // WITH_EDITORONLY_DATA
+#endif	  // USE_USD_SDK
+}	 // namespace UE::InterchangeUsdPipeline::Private
+
+UInterchangeUsdPipeline::UInterchangeUsdPipeline()
+	: GeometryPurpose((int32)(EUsdPurpose::Default | EUsdPurpose::Proxy | EUsdPurpose::Render | EUsdPurpose::Guide))
+	, bImportPseudoRoot(true)
+	, bGeneratePrimvarCompatibleMaterials(true)
+{
+}
+
+void UInterchangeUsdPipeline::ExecutePostFactoryPipeline(
+	const UInterchangeBaseNodeContainer* BaseNodeContainer,
+	const FString& NodeKey,
+	UObject* CreatedAsset,
+	bool bIsAReimport
+)
+{
+#if USE_USD_SDK && WITH_EDITORONLY_DATA
+	if (!BaseNodeContainer || !CreatedAsset)
+	{
+		return;
+	}
+
+	const UInterchangeFactoryBaseNode* FactoryNode = BaseNodeContainer->GetFactoryNode(NodeKey);
+	if (!FactoryNode)
+	{
+		return;
+	}
+
+	if (UStaticMesh* StaticMeshAsset = Cast<UStaticMesh>(CreatedAsset))
+	{
+		using namespace UE::InterchangeUsdPipeline::Private;
+		FUsdPrimvarBaker Baker(StaticMeshAsset, BaseNodeContainer, FactoryNode);
+		Baker.Bake();
+	}
+#endif	  // USE_USD_SDK && WITH_EDITORONLY_DATA
+}
+
+void UInterchangeUsdPipeline::ExecutePipeline(
+	UInterchangeBaseNodeContainer* NodeContainer,
+	const TArray<UInterchangeSourceData*>& InSourceDatas,
+	const FString& ContentBasePath
+)
+{
+	using namespace UE::InterchangeUsdPipeline::Private;
+
+	Super::ExecutePipeline(NodeContainer, InSourceDatas, ContentBasePath);
+
+#if USE_USD_SDK
+	if (!NodeContainer)
+	{
+		return;
+	}
+
+	DisableNodesBasedOnGeometryPurpose(*NodeContainer, (EUsdPurpose)GeometryPurpose);
+
+	RepairMaterialInstanceFactoryNodes(*NodeContainer);
+
+	ProcessVolumeNodes(*NodeContainer);
+
+	if (!bImportPseudoRoot)
+	{
+		RemovePseudoRootFactoryNode(*NodeContainer);
+	}
+
+	ProcessMeshNodes(NodeContainer, bGeneratePrimvarCompatibleMaterials);
+#endif	  // USE_USD_SDK
+}
+
+bool UInterchangeUsdPipeline::CanExecuteOnAnyThread(EInterchangePipelineTask PipelineTask)
+{
+	// TODO: Only the volume node material handling needs to be on the game thread. Maybe that can be split into
+	// a separate USDGameThreadPipeline?
+	return false;
+}
diff --git a/Engine/Plugins/Interchange/Extensions/OpenUSD/Source/Import/Private/InterchangeUsdContext.cpp b/Engine/Plugins/Interchange/Extensions/OpenUSD/Source/Import/Private/InterchangeUsdContext.cpp
new file mode 100644
--- /dev/null
+++ b/Engine/Plugins/Interchange/Extensions/OpenUSD/Source/Import/Private/InterchangeUsdContext.cpp
@@ -0,0 +1,115 @@
+// Copyright Epic Games, Inc. All Rights Reserved.
+
+#include "InterchangeUsdContext.h"
+
+#include "Objects/USDInfoCache.h"
+#include "USDConversionUtils.h"
+#include "UsdWrappers/UsdStage.h"
+
+#include UE_INLINE_GENERATED_CPP_BY_NAME(InterchangeUsdContext)
+
+#define LOCTEXT_NAMESPACE "InterchangeUSDContext"
+
+UInterchangeUsdContext::UInterchangeUsdContext()
+{
+}
+
+void UInterchangeUsdContext::BeginDestroy()
+{
+#if USE_USD_SDK
+	if (bShouldCleanUpFromStageCache && StageIdInUsdUtilsStageCache != INDEX_NONE)
+	{
+		UsdUtils::RemoveStageFromUsdUtilsStageCache(StageIdInUsdUtilsStageCache);
+	}
+
+	ReleaseInfoCache();
+#endif	  // USE_USD_SDK
+
+	Super::BeginDestroy();
+}
+
+int64 UInterchangeUsdContext::GetStageId() const
+{
+	return StageIdInUsdUtilsStageCache;
+}
+
+void UInterchangeUsdContext::SetStageId(int64 InStageId)
+{
+	StageIdInUsdUtilsStageCache = InStageId;
+}
+
+UE::FUsdStage UInterchangeUsdContext::GetUsdStage() const
+{
+#if USE_USD_SDK
+	return UsdUtils::FindUsdUtilsStageCacheStageId(StageIdInUsdUtilsStageCache);
+#else
+	return {};
+#endif	  // USE_USD_SDK
+}
+
+bool UInterchangeUsdContext::SetUsdStage(const UE::FUsdStage& InStage)
+{
+#if USE_USD_SDK
+	UE::FUsdStage OurStage = GetUsdStage();
+	if (OurStage == InStage)
+	{
+		return true;
+	}
+
+	// This stage is already in the stage cache somehow. Just take its existing id,
+	// but let's remember to not remove it from the stage cache whenever we're done, because
+	// it wasn't us that put it there in the first place
+	int64 ExistingId = UsdUtils::GetUsdUtilsStageCacheStageId(InStage);
+	if (ExistingId != INDEX_NONE)
+	{
+		bShouldCleanUpFromStageCache = false;
+		StageIdInUsdUtilsStageCache = ExistingId;
+		return true;
+	}
+
+	// We're adding this stage to the stage cache, so let's make sure it's cleaned
+	// up whenever we're released or else it will remain there forever
+	bShouldCleanUpFromStageCache = true;
+	StageIdInUsdUtilsStageCache = UsdUtils::InsertStageIntoUsdUtilsStageCache(InStage);
+	return StageIdInUsdUtilsStageCache != INDEX_NONE;
+#else
+	return false;
+#endif	  // USE_USD_SDK
+}
+
+FUsdInfoCache* UInterchangeUsdContext::GetInfoCache() const
+{
+	return InfoCache;
+}
+
+void UInterchangeUsdContext::SetExternalInfoCache(FUsdInfoCache& InInfoCache)
+{
+	if (&InInfoCache == GetInfoCache())
+	{
+		return;
+	}
+
+	ReleaseInfoCache();
+
+	InfoCache = &InInfoCache;
+}
+
+FUsdInfoCache* UInterchangeUsdContext::CreateOwnedInfoCache()
+{
+	if (OwnedInfoCache.IsValid())
+	{
+		return GetInfoCache();
+	}
+
+	OwnedInfoCache = MakeUnique<FUsdInfoCache>();
+	InfoCache = OwnedInfoCache.Get();
+	return GetInfoCache();
+}
+
+void UInterchangeUsdContext::ReleaseInfoCache()
+{
+	InfoCache = nullptr;
+	OwnedInfoCache.Reset();
+}
+
+#undef LOCTEXT_NAMESPACE
diff --git a/Engine/Plugins/Interchange/Extensions/OpenUSD/Source/Import/Private/InterchangeUsdTranslator.cpp b/Engine/Plugins/Interchange/Extensions/OpenUSD/Source/Import/Private/InterchangeUsdTranslator.cpp
new file mode 100644
--- /dev/null
+++ b/Engine/Plugins/Interchange/Extensions/OpenUSD/Source/Import/Private/InterchangeUsdTranslator.cpp
@@ -0,0 +1,4837 @@
+// Copyright Epic Games, Inc. All Rights Reserved.
+
+#include "InterchangeUsdTranslator.h"
+
+#include "Objects/USDInfoCache.h"
+#include "Objects/USDSchemaTranslator.h"
+#include "UnrealUSDWrapper.h"
+#include "USDConversionUtils.h"
+#include "USDGeomMeshConversion.h"
+#include "USDLayerUtils.h"
+#include "USDLightConversion.h"
+#include "USDLog.h"
+#include "USDMaterialXShaderGraph.h"
+#include "USDObjectUtils.h"
+#include "USDPrimConversion.h"
+#include "USDProjectSettings.h"
+#include "USDShadeConversion.h"
+#include "USDSkeletalDataConversion.h"
+#include "USDStageOptions.h"
+#include "USDTypesConversion.h"
+#include "UsdWrappers/SdfPath.h"
+#include "UsdWrappers/UsdAttribute.h"
+#include "UsdWrappers/UsdEditContext.h"
+#include "UsdWrappers/UsdPrim.h"
+#include "UsdWrappers/UsdRelationship.h"
+#include "UsdWrappers/UsdSkelAnimQuery.h"
+#include "UsdWrappers/UsdSkelBinding.h"
+#include "UsdWrappers/UsdSkelBlendShape.h"
+#include "UsdWrappers/UsdSkelBlendShapeQuery.h"
+#include "UsdWrappers/UsdSkelCache.h"
+#include "UsdWrappers/UsdSkelInbetweenShape.h"
+#include "UsdWrappers/UsdSkelSkeletonQuery.h"
+#include "UsdWrappers/UsdSkelSkinningQuery.h"
+#include "UsdWrappers/UsdStage.h"
+#include "UsdWrappers/UsdTyped.h"
+#include "UsdWrappers/UsdVariantSets.h"
+
+#include "InterchangeCameraNode.h"
+#include "InterchangeLightNode.h"
+#include "InterchangeManager.h"
+#include "InterchangeMaterialDefinitions.h"
+#include "InterchangeMaterialInstanceNode.h"
+#include "InterchangeMaterialReferenceNode.h"
+#include "InterchangeMeshNode.h"
+#include "InterchangeSceneNode.h"
+#include "InterchangeShaderGraphNode.h"
+#include "InterchangeTexture2DNode.h"
+#include "InterchangeTranslatorHelper.h"
+#include "InterchangeUsdContext.h"
+#include "InterchangeVolumeNode.h"
+#include "Mesh/InterchangeMeshPayload.h"
+#include "Nodes/InterchangeUserDefinedAttribute.h"
+#include "Usd/InterchangeUsdDefinitions.h"
+#include "Volume/InterchangeVolumeDefinitions.h"
+#include "Volume/InterchangeVolumeTranslatorSettings.h"
+
+#include "MaterialX/MaterialXUtils/MaterialXBase.h"
+#include "MaterialX/MaterialXUtils/MaterialXManager.h"
+
+#include "Async/ParallelFor.h"
+#include "Components/HeterogeneousVolumeComponent.h"
+#include "HAL/IConsoleManager.h"
+#include "Internationalization/Regex.h"
+#include "Modules/ModuleManager.h"
+#include "MovieSceneSection.h"
+#include "Rendering/SkeletalMeshLODImporterData.h"
+#include "StaticMeshAttributes.h"
+#include "UDIMUtilities.h"
+#include "UObject/GCObjectScopeGuard.h"
+
+#if USE_USD_SDK
+#include "USDIncludesStart.h"
+#include "pxr/usd/usdGeom/tokens.h"
+#include "pxr/usd/usdLux/tokens.h"
+#include "pxr/usd/usdShade/tokens.h"
+#include "USDIncludesEnd.h"
+#endif	  // USE_USD_SDK
+
+#include UE_INLINE_GENERATED_CPP_BY_NAME(InterchangeUsdTranslator)
+
+#define LOCTEXT_NAMESPACE "InterchangeUSDTranslator"
+
+static bool GInterchangeEnableUSDImport = false;
+static FAutoConsoleVariableRef CVarInterchangeEnableUSDImport(
+	TEXT("Interchange.FeatureFlags.Import.USD"),
+	GInterchangeEnableUSDImport,
+	TEXT("Whether USD support is enabled.")
+);
+
+static bool GInterchangeEnableUSDLevelImport = false;
+// Import into level via USD Interchange is disabled for 5.5 as it's still a work in progress
+static FAutoConsoleVariableRef CVarInterchangeEnableUSDLevelImport(
+	TEXT("Interchange.FeatureFlags.Import.USD.ToLevel"),
+	GInterchangeEnableUSDLevelImport,
+	TEXT("Whether support for USD level import is enabled.")
+);
+
+namespace UE::InterchangeUsdTranslator::Private
+{
+#if USE_USD_SDK
+	const static FString AnimationPrefix = TEXT("\\Animation\\");
+	const static FString AnimationTrackPrefix = TEXT("\\AnimationTrack\\");
+	const static FString BonePrefix = TEXT("\\Bone\\");
+	const static FString CameraPrefix = TEXT("\\Camera\\");
+	const static FString LightPrefix = TEXT("\\Light\\");
+	const static FString VolumePrefix = TEXT("\\Volume\\");
+	const static FString LODPrefix = TEXT("\\LOD\\");
+	const static FString MaterialPrefix = TEXT("\\Material\\");
+	const static FString MaterialReferencePrefix = TEXT("\\MaterialReference\\");
+	const static FString MeshPrefix = TEXT("\\Mesh\\");
+	const static FString MorphTargetPrefix = TEXT("\\MorphTarget\\");
+	const static FString PrimitiveShapePrefix = TEXT("\\PrimitiveShape\\");
+
+	const static FString LODContainerSuffix = TEXT("LODContainer");
+	const static FString TwoSidedSuffix = TEXT("_TwoSided");
+
+	const static FString LODString = TEXT("LOD");
+
+	// Information intended to be passed down from parent to children (by value) as we traverse the stage
+	struct FTraversalInfo
+	{
+		UInterchangeBaseNode* ParentNode = nullptr;
+
+		TSharedPtr<UE::FUsdSkelCache> FurthestSkelCache;
+		TSharedPtr<FString> FurthestParentSkelRootPath;	   // Used to populate the skel cache
+		TSharedPtr<FString> ClosestParentSkelRootPath;
+		TSharedPtr<FString> BoundSkeletonPrimPath;
+		TSharedPtr<TArray<FString>> SkelJointNames;	   // Needed for skel mesh payloads
+
+		bool bVisible = true;
+		bool bInsideLOD = false;
+		bool bIsLODContainer = false;
+
+	public:
+		void UpdateWithCurrentPrim(const UE::FUsdPrim& CurrentPrim)
+		{
+			bVisible = bVisible && UsdUtils::HasInheritedVisibility(CurrentPrim);
+
+			// Check this first so that we go bInsideLOD if our parent was a LODContainer
+			bInsideLOD = bInsideLOD || bIsLODContainer;
+
+			// We only want this to be true when we're traversing the exact prim that owns
+			// the LOD: Once we step into any of its children it should go back to false
+			bIsLODContainer = CurrentPrim.GetVariantSets().HasVariantSet(LODString);
+
+			if (CurrentPrim.IsA(TEXT("SkelRoot")))
+			{
+				if (!ClosestParentSkelRootPath)
+				{
+					// The root-most skel cache should handle any nested UsdSkel prims as well
+					FurthestSkelCache = MakeShared<UE::FUsdSkelCache>();
+
+					const bool bTraverseInstanceProxies = true;
+					FurthestSkelCache->Populate(CurrentPrim, bTraverseInstanceProxies);
+					FurthestParentSkelRootPath = MakeShared<FString>(CurrentPrim.GetPrimPath().GetString());
+				}
+
+				ClosestParentSkelRootPath = MakeShared<FString>(CurrentPrim.GetPrimPath().GetString());
+			}
+
+			if (ClosestParentSkelRootPath && CurrentPrim.HasAPI(TEXT("SkelBindingAPI")))
+			{
+				UE::FUsdStage Stage = CurrentPrim.GetStage();
+
+				if (UE::FUsdRelationship SkelRel = CurrentPrim.GetRelationship(TEXT("skel:skeleton")))
+				{
+					TArray<UE::FSdfPath> Targets;
+					if (SkelRel.GetTargets(Targets) && Targets.Num() > 0)
+					{
+						UE::FUsdPrim TargetSkeleton = Stage.GetPrimAtPath(Targets[0]);
+						if (TargetSkeleton && TargetSkeleton.IsA(TEXT("Skeleton")))
+						{
+							BoundSkeletonPrimPath = MakeShared<FString>(TargetSkeleton.GetPrimPath().GetString());
+						}
+					}
+				}
+			}
+		}
+
+		UE::FUsdSkelSkeletonQuery ResolveSkelQuery(const UE::FUsdStage& Stage) const
+		{
+			if (!Stage || !BoundSkeletonPrimPath || BoundSkeletonPrimPath->IsEmpty())
+			{
+				return {};
+			}
+
+			UE::FUsdPrim SkeletonPrim = Stage.GetPrimAtPath(UE::FSdfPath{**BoundSkeletonPrimPath});
+			if (!SkeletonPrim)
+			{
+				return {};
+			}
+
+			return FurthestSkelCache->GetSkelQuery(SkeletonPrim);
+		}
+
+		UE::FUsdPrim ResolveClosestParentSkelRoot(const UE::FUsdStage& Stage) const
+		{
+			if (!Stage || !ClosestParentSkelRootPath || ClosestParentSkelRootPath->IsEmpty())
+			{
+				return {};
+			}
+
+			return Stage.GetPrimAtPath(UE::FSdfPath{**ClosestParentSkelRootPath});
+		}
+
+		void RepopulateSkelCache(const UE::FUsdStage& Stage)
+		{
+			if (!FurthestSkelCache.IsValid() || !FurthestParentSkelRootPath.IsValid() || FurthestParentSkelRootPath->IsEmpty())
+			{
+				return;
+			}
+
+			UE::FUsdPrim SkelRootPrim = Stage.GetPrimAtPath(UE::FSdfPath{**FurthestParentSkelRootPath});
+			if (!SkelRootPrim)
+			{
+				return;
+			}
+
+			const bool bTraverseInstanceProxies = true;
+			ensure(FurthestSkelCache->Populate(SkelRootPrim, bTraverseInstanceProxies));
+		}
+	};
+
+	void Traverse(
+		const UE::FUsdPrim& Prim,
+		UInterchangeUSDTranslatorImpl& TranslatorImpl,
+		UInterchangeBaseNodeContainer& NodeContainer,
+		UInterchangeUsdTranslatorSettings* TranslatorSettings,
+		FTraversalInfo Info
+	);
+
+	// clang-format off
+	const static TMap<FName, EInterchangePropertyTracks> PropertyNameToTrackType = {
+		// Common properties
+		{UnrealIdentifiers::HiddenInGamePropertyName, 			EInterchangePropertyTracks::ActorHiddenInGame}, // Binding visibility to the actor works better for cameras
+
+		// Camera properties
+		// TODO: Need to add support for SensorHorizontalOffset and SensorVerticalOffset, once Interchange supports those
+		{UnrealIdentifiers::CurrentFocalLengthPropertyName, 		EInterchangePropertyTracks::CameraCurrentFocalLength},
+		{UnrealIdentifiers::ManualFocusDistancePropertyName, 		EInterchangePropertyTracks::CameraFocusSettingsManualFocusDistance},
+		{UnrealIdentifiers::CurrentAperturePropertyName, 			EInterchangePropertyTracks::CameraCurrentAperture},
+		{UnrealIdentifiers::SensorWidthPropertyName, 				EInterchangePropertyTracks::CameraFilmbackSensorWidth},
+		{UnrealIdentifiers::SensorHeightPropertyName, 				EInterchangePropertyTracks::CameraFilmbackSensorHeight},
+		{UnrealIdentifiers::ExposureCompensationPropertyName, 		EInterchangePropertyTracks::CameraPostProcessSettingsAutoExposureBias},
+		{UnrealIdentifiers::ProjectionModePropertyName, 			EInterchangePropertyTracks::CameraProjectionMode},
+		{UnrealIdentifiers::OrthoFarClipPlanePropertyName, 			EInterchangePropertyTracks::CameraOrthoFarClipPlane},
+		{UnrealIdentifiers::OrthoNearClipPlanePropertyName, 		EInterchangePropertyTracks::CameraOrthoNearClipPlane},
+		{UnrealIdentifiers::CustomNearClipppingPlanePropertyName,	EInterchangePropertyTracks::CameraCustomNearClippingPlane},
+
+		// Light properties
+		{UnrealIdentifiers::LightColorPropertyName, 		EInterchangePropertyTracks::LightColor},
+		{UnrealIdentifiers::TemperaturePropertyName, 		EInterchangePropertyTracks::LightTemperature},
+		{UnrealIdentifiers::UseTemperaturePropertyName, 	EInterchangePropertyTracks::LightUseTemperature},
+		{UnrealIdentifiers::SourceHeightPropertyName, 		EInterchangePropertyTracks::LightSourceHeight},
+		{UnrealIdentifiers::SourceWidthPropertyName, 		EInterchangePropertyTracks::LightSourceWidth},
+		{UnrealIdentifiers::SourceRadiusPropertyName, 		EInterchangePropertyTracks::LightSourceRadius},
+		{UnrealIdentifiers::OuterConeAnglePropertyName, 	EInterchangePropertyTracks::LightOuterConeAngle},
+		{UnrealIdentifiers::InnerConeAnglePropertyName, 	EInterchangePropertyTracks::LightInnerConeAngle},
+		{UnrealIdentifiers::LightSourceAnglePropertyName, 	EInterchangePropertyTracks::LightSourceAngle},
+		{UnrealIdentifiers::IntensityPropertyName, 			EInterchangePropertyTracks::LightIntensity},
+	};
+	// clang-format on
+
+	// Small container that we can use Pimpl with so we don't have to include too many USD includes on the header file.
+	//
+	// It also skirts around a small complication where UInterchangeUSDTranslator::Translate is const, and yet we must
+	// keep and modify some members (like UsdStage) for when the payload functions get called later... The other translators
+	// use mutable or const casts, but with the Impl we don't need to!
+	class UInterchangeUSDTranslatorImpl
+	{
+	public:
+		/** Add a material instance to the node container, otherwise it will add a material if it comes from a Translator (for example coming from
+		 * MaterialX which cannot handle material instances) */
+		TArray<FString> AddMaterialNode(
+			const UE::FUsdPrim& Prim,
+			UInterchangeUsdTranslatorSettings* TranslatorSettings,
+			UInterchangeBaseNodeContainer& NodeContainer
+		);
+
+		FString AddMeshNode(
+			const UE::FUsdPrim& Prim,
+			UInterchangeBaseNodeContainer& NodeContainer,
+			const FTraversalInfo& Info,
+			bool bPrimitiveShape = false
+		);
+
+		void AddLODMeshNodes(
+			const UE::FUsdPrim& Prim,
+			UInterchangeBaseNodeContainer& NodeContainer,
+			UInterchangeSceneNode& LODContainerSceneNode,
+			UInterchangeUsdTranslatorSettings* TranslatorSettings,
+			FTraversalInfo Info
+		);
+
+		TArray<FString> AddVolumeNodes(
+			const UE::FUsdPrim& InPrim,
+			UInterchangeBaseNodeContainer& InNodeContainer,
+			FString& OutMaterialInstanceUid,
+			bool& bOutNeedsFrameTrack
+		);
+
+		/**
+		 * Search for a MaterialX file embedded in the file, in case there's none, create a ShaderGraph using MaterialX
+		 * @param GeomPropValueNames - Array of the <geompropvalue> node names from the shader graph that has been converted to <image> nodes (used
+		 * for baking later on by the Factory)
+		 */
+		bool AddMaterialXShaderGraph(
+			const UE::FUsdPrim& InPrim,
+			UInterchangeUsdTranslatorSettings* TranslatorSettings,
+			UInterchangeBaseNodeContainer& NodeContainer,
+			TArray<FUsdMaterialXShaderGraph::FGeomProp>& OutGeomPropValues
+		);
+
+		/**
+		 * Set the material slot dependencies on the meshes, if bIterateShaderGraphs is true, it will iterate over the shader graphs to rename them
+		 * according to the MaterialPrimName
+		 */
+		void SetMaterialSlotDependencies(
+			const FString& MaterialPrimName,
+			const FString& MaterialUid,
+			UInterchangeBaseNodeContainer& NodeContainer,
+			const TArray<FUsdMaterialXShaderGraph::FGeomProp>& GeomPropValues = {},
+			bool bIterateShaderGraphs = false
+		);
+
+		/**
+		 * If we're not translating a decompressed USD root, returns TexturePathOnDisk.
+		 * If we are translating a decompressed USD root, returns the path to the USDZ file itself.
+		 *
+		 * The intent here is that in the USDZ case the texture filepath will point at a temp file on disk, that we may dispose
+		 * of later after importing. In order to allow reimporting the texture at a later time, we'll just put the USDZ path
+		 * itself as it's source path, and tweak the USD translator to know what to do with this
+		 */
+		FString GetTextureSourcePath(const FString& TexturePathOnDisk);
+
+		/** If we decompressed a USDZ file to a temp folder this will delete everything from that folder */
+		void CleanUpDecompressedUSDZFolder();
+
+		void SetupTranslationContext(const UInterchangeUsdTranslatorSettings& Settings);
+
+		void SetupInfoCache();
+
+		UE::FUsdPrim TryGettingInactiveLODPrim(const FString& PrimPath);
+
+		void FinalizeLODContainerTraversal(
+			UInterchangeBaseNodeContainer& NodeContainer,
+			const FTraversalInfo& Info,
+			const UInterchangeSceneNode* SceneNodeWithLODs
+		);
+
+	public:
+		// We have to keep a stage reference so that we can parse the payloads after Translate() complete.
+		// ReleaseSource() clears this member, once translation is complete.
+		UE::FUsdStage UsdStage;
+		TSharedPtr<FUsdSchemaTranslationContext> TranslationContext;
+		FUsdInfoCache* InfoCache = nullptr;
+
+		// We store temp stages in here that we open in order to parse stuff inside of inactive variants of our main UsdStage.
+		// We do this because the payload data are retrieved concurrently, and toggling variants mutates the current stage.
+		mutable FRWLock PrimPathToVariantToStageLock;
+		TMap<FString, TMap<FString, UE::FUsdStage>> PrimPathToVariantToStage;
+
+		FString USDZFilePath;
+		FString DecompressedUSDZRoot;
+
+		// On UInterchangeUSDTranslator::Translate we set this up based on our TranslatorSettings, and then
+		// we can reuse it (otherwise we have to keep converting the FNames into Tokens all the time)
+		UsdToUnreal::FUsdMeshConversionOptions CachedMeshConversionOptions;
+
+		TMap<FString, UsdUtils::FUsdPrimMaterialAssignmentInfo> CachedMaterialAssignments;
+
+		// We fill this in while we're translating a LOD container so that we can do some post-processing inside FinalizeLODContainerTraversal
+		TArray<UInterchangeSceneNode*> CurrentLODSceneNodes;
+
+		// When traversing we'll generate FTraversalInfo objects. If we need to (e.g. for skinned meshes),
+		// we'll store the info for that translated node here, so we don't have to recompute it when returning
+		// the payload data.
+		// Note: We only do this when needed: This shouldn't have data for every prim in the stage.
+		TMap<FString, FTraversalInfo> NodeUidToCachedTraversalInfo;
+		mutable FRWLock CachedTraversalInfoLock;
+
+		// This node eventually becomes a LevelSequence, and all track nodes are connected to it.
+		// For now we only generate a single LevelSequence per stage though, so we'll keep track of this
+		// here for easy access when parsing the tracks
+		UInterchangeAnimationTrackSetNode* CurrentTrackSet = nullptr;
+
+		// Map of translators that we call in the GetTexturePayload, the key has no real meaning, it's just here to avoid having duplicates and
+		// calling several times the Translate function
+		TMap<FString, TStrongObjectPtr<UInterchangeTranslatorBase>> Translators;
+
+		// We stash the info we collected from each Volume prim path here, as we'll reuse it between translation and retrieving the payloads
+		TMap<FString, TArray<UsdUtils::FVolumePrimInfo>> PrimPathToVolumeInfo;
+
+	private:
+		struct FMaterialSlotMesh
+		{
+			FString MaterialSlotName;
+			UInterchangeMeshNode* MeshNode;
+		};
+		TMap<FString, FString> MaterialUidToActualNodeUid;
+		TMap<FString, TArray<FMaterialSlotMesh>> PrimPathToSlotMeshNodes;
+
+		// Used within a translation. We cache these because we make a volume node *per .vdb file*, and on the USD side we may have
+		// any number of Volume prims internally using the same .vdb file, and we want to share these whenever possible
+		TMap<FString, TMap<FString, UInterchangeVolumeNode*>> VolumeFilepathToAnimationIDToNode;
+	};
+
+	// Adds a numbered suffix (if needed) to NodeUid to make sure there is nothing with that ID within NodeContainer.
+	// Does not add the ID to the NodeContainer.
+	void MakeNodeUidUniqueInContainer(FString& NodeUid, const UInterchangeBaseNodeContainer& NodeContainer)
+	{
+		if (!NodeContainer.IsNodeUidValid(NodeUid) || NodeUid == UInterchangeBaseNode::InvalidNodeUid())
+		{
+			return;
+		}
+
+		int32 Suffix = 0;
+		FString Result;
+		do
+		{
+			Result = FString::Printf(TEXT("%s_%d"), *NodeUid, Suffix++);
+		} while (NodeContainer.IsNodeUidValid(Result));
+
+		NodeUid = Result;
+	}
+
+	bool IsValidLODName(const FString& PrimName)
+	{
+		return PrimName.Len() > LODString.Len() && PrimName.StartsWith(LODString) && PrimName.RightChop(LODString.Len()).IsNumeric();
+	}
+
+	int32 GetLODIndexFromName(FString Name)
+	{
+		if (!Name.RemoveFromStart(LODString))
+		{
+			return INDEX_NONE;
+		}
+
+		if (!Name.IsNumeric())
+		{
+			return INDEX_NONE;
+		}
+
+		int32 Index = INDEX_NONE;
+		LexFromString(Index, *Name);
+		return Index;
+	}
+
+	UE::FUsdPrim GetLODMesh(const UE::FUsdPrim& LODContainerPrim, const FString& LODName)
+	{
+		const UE::FSdfPath IdealMeshPrimPath = LODContainerPrim.GetPrimPath().AppendChild(*LODName);
+
+		UE::FUsdPrim Prim = LODContainerPrim.GetStage().GetPrimAtPath(IdealMeshPrimPath);
+		if (Prim.IsA(TEXT("Mesh")))
+		{
+			return Prim;
+		}
+
+		return {};
+	};
+
+	bool CheckAndChopPayloadPrefix(FString& PayloadKey, const FString& Prefix)
+	{
+		if (PayloadKey.StartsWith(Prefix))
+		{
+			// Remove the prefix from the payload key
+			PayloadKey = PayloadKey.RightChop(Prefix.Len());
+
+			return true;
+		}
+
+		return false;
+	}
+
+	FString HashAnimPayloadQuery(const Interchange::FAnimationPayloadQuery& Query)
+	{
+		FSHAHash Hash;
+		FSHA1 SHA1;
+
+		// TODO: Is there a StringView alternative?
+		FString SkeletonPrimPath;
+		FString JointIndexStr;
+		bool bSplit = Query.PayloadKey.UniqueId.Split(TEXT("\\"), &SkeletonPrimPath, &JointIndexStr, ESearchCase::CaseSensitive, ESearchDir::FromEnd);
+		if (!bSplit)
+		{
+			return {};
+		}
+
+		SHA1.UpdateWithString(*SkeletonPrimPath, SkeletonPrimPath.Len());
+
+		SHA1.Update(reinterpret_cast<const uint8*>(&Query.TimeDescription.BakeFrequency), sizeof(Query.TimeDescription.BakeFrequency));
+		SHA1.Update(reinterpret_cast<const uint8*>(&Query.TimeDescription.RangeStartSecond), sizeof(Query.TimeDescription.RangeStartSecond));
+		SHA1.Update(reinterpret_cast<const uint8*>(&Query.TimeDescription.RangeStopSecond), sizeof(Query.TimeDescription.RangeStopSecond));
+
+		SHA1.Final();
+		SHA1.GetHash(&Hash.Hash[0]);
+		return Hash.ToString();
+	}
+
+	FString GetMorphTargetMeshNodeUid(const FString& MeshPrimPath, int32 MeshBlendShapeIndex, const FString& InbetweenName = FString{})
+	{
+		return FString::Printf(TEXT("%s%s\\%d\\%s"), *MorphTargetPrefix, *MeshPrimPath, MeshBlendShapeIndex, *InbetweenName);
+	}
+
+	FString GetMorphTargetMeshPayloadKey(
+		bool bIsInsideLOD,
+		const FString& MeshPrimPath,
+		int32 MeshBlendShapeIndex,
+		const FString& InbetweenName = FString{}
+	)
+	{
+		return FString::Printf(TEXT("%s%s\\%d\\%s"), bIsInsideLOD ? *LODPrefix : TEXT(""), *MeshPrimPath, MeshBlendShapeIndex, *InbetweenName);
+	}
+
+	// TODO: Cleanup/unify/standardize these payload manipulating functions (don't add/remove prefixes everywhere but have a standard format, etc.)
+	bool ParseMorphTargetMeshPayloadKey(
+		FString InPayloadKey,
+		bool& bOutIsLODMesh,
+		FString& OutMeshPrimPath,
+		int32& OutBlendShapeIndex,
+		FString& OutInbetweenName
+	)
+	{
+		bool bIsLODMesh = CheckAndChopPayloadPrefix(InPayloadKey, LODPrefix);
+
+		// These payload keys are generated by GetMorphTargetMeshPayloadKey(), and so should take the form
+		// "<mesh prim path>\<mesh blend shape index>\<optional inbetween name>"
+		const bool bCullEmpty = false;
+		TArray<FString> PayloadKeyTokens;
+		InPayloadKey.ParseIntoArray(PayloadKeyTokens, TEXT("\\"), bCullEmpty);
+		if (PayloadKeyTokens.Num() != 3)
+		{
+			return false;
+		}
+
+		const FString& MeshPrimPath = PayloadKeyTokens[0];
+		const FString& BlendShapeIndexStr = PayloadKeyTokens[1];
+		const FString& InbetweenName = PayloadKeyTokens[2];
+
+		int32 BlendShapeIndex = INDEX_NONE;
+		bool bLexed = LexTryParseString(BlendShapeIndex, *BlendShapeIndexStr);
+		if (!bLexed)
+		{
+			return false;
+		}
+
+		bOutIsLODMesh = bIsLODMesh;
+		OutMeshPrimPath = MeshPrimPath;
+		OutBlendShapeIndex = BlendShapeIndex;
+		OutInbetweenName = InbetweenName;
+		return true;
+	}
+
+	FString GetMorphTargetCurvePayloadKey(const FString& SkeletonPrimPath, int32 SkelAnimChannelIndex, const FString& BlendShapePath)
+	{
+		return FString::Printf(TEXT("%s\\%d\\%s"), *SkeletonPrimPath, SkelAnimChannelIndex, *BlendShapePath);
+	}
+
+	FString EncodeTexturePayloadKey(const UsdToUnreal::FTextureParameterValue& Value)
+	{
+		// Encode the compression settings onto the payload key as we need to move that into the
+		// payload data within UInterchangeUSDTranslator::GetTexturePayloadData.
+		//
+		// This should be a temporary thing, and in the future we'll be able to store compression
+		// settings directly on the texture translated node
+		return Value.TextureFilePath + TEXT("\\") + LexToString((int32)Value.Group);
+	}
+
+	bool DecodeTexturePayloadKey(const FString& PayloadKey, FString& OutTextureFilePath, TextureGroup& OutTextureGroup)
+	{
+		// Use split from end here so that we ignore any backslashes within the file path itself
+		FString FilePath;
+		FString TextureGroupStr;
+		bool bSplit = PayloadKey.Split(TEXT("\\"), &FilePath, &TextureGroupStr, ESearchCase::CaseSensitive, ESearchDir::FromEnd);
+		if (!bSplit)
+		{
+			return false;
+		}
+
+		OutTextureFilePath = FilePath;
+
+		int32 TempInt;
+		if (LexTryParseString<int32>(TempInt, *TextureGroupStr))
+		{
+			OutTextureGroup = (TextureGroup)TempInt;
+		}
+
+		return true;
+	}
+
+	void FixMaterialSlotNames(FMeshDescription& MeshDescription, const TArray<UsdUtils::FUsdPrimMaterialSlot>& MeshAssingmentSlots)
+	{
+		// Fixup material slot names to match the material that is assigned. For Interchange it is better to have the material
+		// slot names match what is assigned into them, as it will use those names to "merge identical slots" depending on the
+		// import options.
+		//
+		// Note: These names must also match what is set via MeshNode->SetSlotMaterialDependencyUid(SlotName, MaterialUid)
+		FStaticMeshAttributes StaticMeshAttributes(MeshDescription);
+		for (int32 MaterialSlotIndex = 0; MaterialSlotIndex < StaticMeshAttributes.GetPolygonGroupMaterialSlotNames().GetNumElements();
+			 ++MaterialSlotIndex)
+		{
+			int32 MaterialIndex = 0;
+			LexFromString(MaterialIndex, *StaticMeshAttributes.GetPolygonGroupMaterialSlotNames()[MaterialSlotIndex].ToString());
+
+			if (MeshAssingmentSlots.IsValidIndex(MaterialIndex))
+			{
+				const FString Source = MeshAssingmentSlots[MaterialIndex].MaterialSource;
+				StaticMeshAttributes.GetPolygonGroupMaterialSlotNames()[MaterialSlotIndex] = *Source;
+			}
+		}
+	}
+
+	bool ReadBools(
+		const UE::FUsdStage& UsdStage,
+		const TArray<double>& UsdTimeSamples,
+		const TFunction<bool(double)>& ReaderFunc,
+		Interchange::FAnimationPayloadData& OutPayloadData
+	)
+	{
+		OutPayloadData.StepCurves.SetNum(1);
+		FInterchangeStepCurve& Curve = OutPayloadData.StepCurves[0];
+		TArray<float>& KeyTimes = Curve.KeyTimes;
+		TArray<bool>& BooleanKeyValues = Curve.BooleanKeyValues.Emplace();
+
+		KeyTimes.Reserve(UsdTimeSamples.Num());
+		BooleanKeyValues.Reserve(UsdTimeSamples.Num());
+
+		const FFrameRate StageFrameRate{static_cast<uint32>(UsdStage.GetTimeCodesPerSecond()), 1};
+
+		double LastTimeSample = TNumericLimits<double>::Lowest();
+		for (const double UsdTimeSample : UsdTimeSamples)
+		{
+			// We never want to evaluate the same time twice
+			if (FMath::IsNearlyEqual(UsdTimeSample, LastTimeSample))
+			{
+				continue;
+			}
+			LastTimeSample = UsdTimeSample;
+
+			int32 FrameNumber = FMath::FloorToInt(UsdTimeSample);
+			float SubFrameNumber = UsdTimeSample - FrameNumber;
+
+			FFrameTime FrameTime{FrameNumber, SubFrameNumber};
+			double FrameTimeSeconds = static_cast<float>(StageFrameRate.AsSeconds(FrameTime));
+
+			bool UEValue = ReaderFunc(UsdTimeSample);
+
+			KeyTimes.Add(FrameTimeSeconds);
+			BooleanKeyValues.Add(UEValue);
+		}
+
+		return true;
+	}
+
+	bool ReadFloats(
+		const UE::FUsdStage& UsdStage,
+		const TArray<double>& UsdTimeSamples,
+		const TFunction<float(double)>& ReaderFunc,
+		Interchange::FAnimationPayloadData& OutPayloadData
+	)
+	{
+		OutPayloadData.Curves.SetNum(1);
+		FRichCurve& Curve = OutPayloadData.Curves[0];
+
+		const FFrameRate StageFrameRate{static_cast<uint32>(UsdStage.GetTimeCodesPerSecond()), 1};
+		const ERichCurveInterpMode InterpMode = (UsdStage.GetInterpolationType() == EUsdInterpolationType::Linear)
+													? ERichCurveInterpMode::RCIM_Linear
+													: ERichCurveInterpMode::RCIM_Constant;
+
+		double LastTimeSample = TNumericLimits<double>::Lowest();
+		for (const double UsdTimeSample : UsdTimeSamples)
+		{
+			// We never want to evaluate the same time twice
+			if (FMath::IsNearlyEqual(UsdTimeSample, LastTimeSample))
+			{
+				continue;
+			}
+			LastTimeSample = UsdTimeSample;
+
+			int32 FrameNumber = FMath::FloorToInt(UsdTimeSample);
+			float SubFrameNumber = UsdTimeSample - FrameNumber;
+
+			FFrameTime FrameTime{FrameNumber, SubFrameNumber};
+			double FrameTimeSeconds = static_cast<float>(StageFrameRate.AsSeconds(FrameTime));
+
+			float UEValue = ReaderFunc(UsdTimeSample);
+
+			FKeyHandle Handle = Curve.AddKey(FrameTimeSeconds, UEValue);
+			Curve.SetKeyInterpMode(Handle, InterpMode);
+		}
+
+		return true;
+	}
+
+	bool ReadColors(
+		const UE::FUsdStage& UsdStage,
+		const TArray<double>& UsdTimeSamples,
+		const TFunction<FLinearColor(double)>& ReaderFunc,
+		Interchange::FAnimationPayloadData& OutPayloadData
+	)
+	{
+		OutPayloadData.Curves.SetNum(4);
+		FRichCurve& RCurve = OutPayloadData.Curves[0];
+		FRichCurve& GCurve = OutPayloadData.Curves[1];
+		FRichCurve& BCurve = OutPayloadData.Curves[2];
+		FRichCurve& ACurve = OutPayloadData.Curves[3];
+
+		const FFrameRate StageFrameRate{static_cast<uint32>(UsdStage.GetTimeCodesPerSecond()), 1};
+		const ERichCurveInterpMode InterpMode = (UsdStage.GetInterpolationType() == EUsdInterpolationType::Linear)
+													? ERichCurveInterpMode::RCIM_Linear
+													: ERichCurveInterpMode::RCIM_Constant;
+
+		double LastTimeSample = TNumericLimits<double>::Lowest();
+		for (const double UsdTimeSample : UsdTimeSamples)
+		{
+			// We never want to evaluate the same time twice
+			if (FMath::IsNearlyEqual(UsdTimeSample, LastTimeSample))
+			{
+				continue;
+			}
+			LastTimeSample = UsdTimeSample;
+
+			int32 FrameNumber = FMath::FloorToInt(UsdTimeSample);
+			float SubFrameNumber = UsdTimeSample - FrameNumber;
+
+			FFrameTime FrameTime{FrameNumber, SubFrameNumber};
+			double FrameTimeSeconds = static_cast<float>(StageFrameRate.AsSeconds(FrameTime));
+
+			FLinearColor UEValue = ReaderFunc(UsdTimeSample);
+
+			FKeyHandle RHandle = RCurve.AddKey(FrameTimeSeconds, UEValue.R);
+			FKeyHandle GHandle = GCurve.AddKey(FrameTimeSeconds, UEValue.G);
+			FKeyHandle BHandle = BCurve.AddKey(FrameTimeSeconds, UEValue.B);
+			FKeyHandle AHandle = ACurve.AddKey(FrameTimeSeconds, UEValue.A);
+
+			RCurve.SetKeyInterpMode(RHandle, InterpMode);
+			GCurve.SetKeyInterpMode(GHandle, InterpMode);
+			BCurve.SetKeyInterpMode(BHandle, InterpMode);
+			ACurve.SetKeyInterpMode(AHandle, InterpMode);
+		}
+
+		return true;
+	}
+
+	bool ReadTransforms(
+		const UE::FUsdStage& UsdStage,
+		const TArray<double>& UsdTimeSamples,
+		const TFunction<FTransform(double)>& ReaderFunc,
+		Interchange::FAnimationPayloadData& OutPayloadData
+	)
+	{
+		OutPayloadData.Curves.SetNum(9);
+		FRichCurve& TransXCurve = OutPayloadData.Curves[0];
+		FRichCurve& TransYCurve = OutPayloadData.Curves[1];
+		FRichCurve& TransZCurve = OutPayloadData.Curves[2];
+		FRichCurve& RotXCurve = OutPayloadData.Curves[3];
+		FRichCurve& RotYCurve = OutPayloadData.Curves[4];
+		FRichCurve& RotZCurve = OutPayloadData.Curves[5];
+		FRichCurve& ScaleXCurve = OutPayloadData.Curves[6];
+		FRichCurve& ScaleYCurve = OutPayloadData.Curves[7];
+		FRichCurve& ScaleZCurve = OutPayloadData.Curves[8];
+
+		const FFrameRate StageFrameRate{static_cast<uint32>(UsdStage.GetTimeCodesPerSecond()), 1};
+		const ERichCurveInterpMode InterpMode = (UsdStage.GetInterpolationType() == EUsdInterpolationType::Linear)
+													? ERichCurveInterpMode::RCIM_Linear
+													: ERichCurveInterpMode::RCIM_Constant;
+
+		double LastTimeSample = TNumericLimits<double>::Lowest();
+		for (const double UsdTimeSample : UsdTimeSamples)
+		{
+			// We never want to evaluate the same time twice
+			if (FMath::IsNearlyEqual(UsdTimeSample, LastTimeSample))
+			{
+				continue;
+			}
+			LastTimeSample = UsdTimeSample;
+
+			int32 FrameNumber = FMath::FloorToInt(UsdTimeSample);
+			float SubFrameNumber = UsdTimeSample - FrameNumber;
+
+			FFrameTime FrameTime{FrameNumber, SubFrameNumber};
+			double FrameTimeSeconds = static_cast<float>(StageFrameRate.AsSeconds(FrameTime));
+
+			FTransform UEValue = ReaderFunc(UsdTimeSample);
+			FVector Location = UEValue.GetLocation();
+			FRotator Rotator = UEValue.Rotator();
+			FVector Scale = UEValue.GetScale3D();
+
+			FKeyHandle HandleTransX = TransXCurve.AddKey(FrameTimeSeconds, Location.X);
+			FKeyHandle HandleTransY = TransYCurve.AddKey(FrameTimeSeconds, Location.Y);
+			FKeyHandle HandleTransZ = TransZCurve.AddKey(FrameTimeSeconds, Location.Z);
+			FKeyHandle HandleRotX = RotXCurve.AddKey(FrameTimeSeconds, Rotator.Roll);
+			FKeyHandle HandleRotY = RotYCurve.AddKey(FrameTimeSeconds, Rotator.Pitch);
+			FKeyHandle HandleRotZ = RotZCurve.AddKey(FrameTimeSeconds, Rotator.Yaw);
+			FKeyHandle HandleScaleX = ScaleXCurve.AddKey(FrameTimeSeconds, Scale.X);
+			FKeyHandle HandleScaleY = ScaleYCurve.AddKey(FrameTimeSeconds, Scale.Y);
+			FKeyHandle HandleScaleZ = ScaleZCurve.AddKey(FrameTimeSeconds, Scale.Z);
+
+			TransXCurve.SetKeyInterpMode(HandleTransX, InterpMode);
+			TransYCurve.SetKeyInterpMode(HandleTransY, InterpMode);
+			TransZCurve.SetKeyInterpMode(HandleTransZ, InterpMode);
+			RotXCurve.SetKeyInterpMode(HandleRotX, InterpMode);
+			RotYCurve.SetKeyInterpMode(HandleRotY, InterpMode);
+			RotZCurve.SetKeyInterpMode(HandleRotZ, InterpMode);
+			ScaleXCurve.SetKeyInterpMode(HandleScaleX, InterpMode);
+			ScaleYCurve.SetKeyInterpMode(HandleScaleY, InterpMode);
+			ScaleZCurve.SetKeyInterpMode(HandleScaleZ, InterpMode);
+		}
+
+		return true;
+	}
+
+	void AddTextureNode(
+		const UE::FUsdPrim& Prim,
+		const FString& NodeUid,
+		const UsdToUnreal::FTextureParameterValue& Value,
+		UInterchangeBaseNodeContainer& NodeContainer
+	)
+	{
+		FString PrimPath = Prim.GetPrimPath().GetString();
+		FString NodeName{FPaths::GetCleanFilename(Value.TextureFilePath)};
+
+		// Check if Node already exist with this ID
+		if (const UInterchangeTexture2DNode* Node = Cast<UInterchangeTexture2DNode>(NodeContainer.GetNode(NodeUid)))
+		{
+			return;
+		}
+
+		UInterchangeTexture2DNode* Node = NewObject<UInterchangeTexture2DNode>(&NodeContainer);
+		Node->InitializeNode(NodeUid, NodeName, EInterchangeNodeContainerType::TranslatedAsset);
+		Node->SetPayLoadKey(EncodeTexturePayloadKey(Value));
+
+		static_assert((int)TextureAddress::TA_Wrap == (int)EInterchangeTextureWrapMode::Wrap);
+		static_assert((int)TextureAddress::TA_Clamp == (int)EInterchangeTextureWrapMode::Clamp);
+		static_assert((int)TextureAddress::TA_Mirror == (int)EInterchangeTextureWrapMode::Mirror);
+		Node->SetCustomWrapU((EInterchangeTextureWrapMode)Value.AddressX);
+		Node->SetCustomWrapV((EInterchangeTextureWrapMode)Value.AddressY);
+
+		Node->SetCustomSRGB(Value.GetSRGBValue());
+
+		// Provide the other UDIM tiles
+		//
+		// Note: There is an bImportUDIM option on UInterchangeGenericTexturePipeline that is exclusively used within
+		// UInterchangeGenericTexturePipeline::HandleCreationOfTextureFactoryNode in order to essentially do the exact same
+		// thing as we do here. In theory, we shouldn't need to do this then, and in fact it is a bit bad to do so because
+		// we will always parse these UDIMs whether the option is enabled or disabled. The issue however is that (as of the
+		// time of this writing) UInterchangeGenericTexturePipeline::HandleCreationOfTextureFactoryNode is hard-coded to expect
+		// the texture payload key to be just the texture file path. We can't do that, because we need to also encode
+		// the texture compression settings onto payload key...
+		//
+		// All of that is to say that everything will actually work fine, but if you uncheck "bImportUDIM" on the import options
+		// you will still get UDIMs (for now).
+		if (Value.bIsUDIM)
+		{
+			TMap<int32, FString> TileIndexToPath = UE::TextureUtilitiesCommon::GetUDIMBlocksFromSourceFile(
+				Value.TextureFilePath,
+				UE::TextureUtilitiesCommon::DefaultUdimRegexPattern
+			);
+			Node->SetSourceBlocks(MoveTemp(TileIndexToPath));
+		}
+
+		NodeContainer.AddNode(Node);
+	}
+
+	FString AddDisplayColorMaterialInstanceNodeIfNeeded(UInterchangeBaseNodeContainer& NodeContainer, const UsdUtils::FUsdPrimMaterialSlot& Slot)
+	{
+		using namespace UsdUnreal::MaterialUtils;
+
+		FString DisplayColorDesc = Slot.MaterialSource;
+		FString NodeUid = MaterialPrefix + Slot.MaterialSource;
+
+		if (Slot.bMeshIsDoubleSided)
+		{
+			NodeUid += TwoSidedSuffix;
+			DisplayColorDesc += TwoSidedSuffix;
+		}
+
+		// We'll treat the DisplayColorDesc (something like "!DisplayColor_1_0") as the material instance UID here
+		const UInterchangeShaderGraphNode* Node = Cast<UInterchangeShaderGraphNode>(NodeContainer.GetNode(NodeUid));
+		if (Node)
+		{
+			return NodeUid;
+		}
+
+		// Need to create a new instance
+		TOptional<FDisplayColorMaterial> ParsedMat = FDisplayColorMaterial::FromString(DisplayColorDesc);
+		if (!ParsedMat)
+		{
+			return {};
+		}
+		FString NodeName = ParsedMat->ToPrettyString();
+
+		UInterchangeShaderGraphNode* NewNode = NewObject<UInterchangeShaderGraphNode>(&NodeContainer);
+		NewNode->InitializeNode(NodeUid, NodeName, EInterchangeNodeContainerType::TranslatedAsset);
+		NewNode->AddBooleanAttribute(UE::Interchange::USD::ParseMaterialIdentifier, true);
+		NewNode->AddStringAttribute(*UE::Interchange::USD::DisplayColorDescIdentifier, DisplayColorDesc);
+
+		if (Slot.bMeshIsDoubleSided)
+		{
+			NewNode->SetCustomTwoSided(true);
+			NewNode->SetCustomTwoSidedTransmission(true); // weirdly enough we also have to set this in case of a translucent material
+		}
+		NodeContainer.AddNode(NewNode);
+
+		return NodeUid;
+	}
+
+	FString AddUnrealMaterialReferenceNodeIfNeeded(UInterchangeBaseNodeContainer& NodeContainer, const FString& ContentPath)
+	{
+		using namespace UsdUnreal::MaterialUtils;
+
+		// e.g. "\\MaterialReference\\/Game/MyFolder/Red.Red"
+		const FString NodeUid = MaterialReferencePrefix + ContentPath;
+
+		const UInterchangeMaterialReferenceNode* Node = Cast<UInterchangeMaterialReferenceNode>(NodeContainer.GetNode(NodeUid));
+		if (Node)
+		{
+			return NodeUid;
+		}
+
+		const FString DisplayName = FPaths::GetBaseFilename(ContentPath);
+
+		UInterchangeMaterialReferenceNode* NewNode = NewObject<UInterchangeMaterialReferenceNode>(&NodeContainer);
+		NewNode->InitializeNode(NodeUid, DisplayName, EInterchangeNodeContainerType::TranslatedAsset);
+		NewNode->SetCustomContentPath(ContentPath);
+
+		NodeContainer.AddNode(NewNode);
+
+		return NodeUid;
+	}
+
+	void UInterchangeUSDTranslatorImpl::SetMaterialSlotDependencies(
+		const FString& MaterialPrimName,
+		const FString& MaterialUid,
+		UInterchangeBaseNodeContainer& NodeContainer,
+		const TArray<FUsdMaterialXShaderGraph::FGeomProp>& GeomPropValues,
+		bool bIterateShaderGraphs
+	)
+	{
+		using namespace UE::Interchange;
+		if (bIterateShaderGraphs)
+		{
+			// The material from the MaterialXTranslator doesn't have the same UID, both the PrimPath have the same name but not the same path
+			// We need to retrieve that name (which is the Material name) in the translator, then we can map it to the right mesh
+			NodeContainer.BreakableIterateNodesOfType<UInterchangeShaderGraphNode>(
+				[this, &MaterialPrimName, &MaterialUid](const FString&, UInterchangeShaderGraphNode* ShaderGraphNode)
+				{
+					FString UID = ShaderGraphNode->GetUniqueID();
+					if (FPaths::GetBaseFilename(UID) == MaterialPrimName)
+					{
+						MaterialUidToActualNodeUid.Add(MaterialUid, UID);
+						return true;
+					}
+					else
+					{
+						return false;
+					}
+				}
+			);
+		}
+
+		// Let's iterate over the Shader Nodes, the TextureSample ones more specifically,
+		// to see if have to retrieve any attributes related to the conversion of geompropvalues
+#if WITH_EDITOR
+		// Let's iterate over the Shader Nodes, the TextureSample ones, more specifically,
+		// to see if we have to retrieve any attributes related to the conversion of geompropvalues
+		// we'll store the UID of the shader node in order to retrieve it during the baking phase (in the post factory import)
+		TArray<FString> ShaderNodesTextureSampleUIDs;
+		NodeContainer.IterateNodesOfType<UInterchangeShaderNode>(
+			[&ShaderNodesTextureSampleUIDs](const FString&, UInterchangeShaderNode* ShaderNode)
+			{
+				using namespace UE::Interchange::Materials::Standard::Nodes;
+
+				constexpr const TCHAR* GeomPropHandledAttribute = TEXT("GeomPropHandled");
+				if (FString ShaderType; ShaderNode->GetCustomShaderType(ShaderType) && ShaderType == TextureSample::Name.ToString())
+				{
+					if (bool bIsGeomProp; ShaderNode->GetBooleanAttribute(::MaterialX::Attributes::GeomPropImage, bIsGeomProp) && bIsGeomProp)
+					{
+						// We only want the geomprops that were not already assigned to the mesh node
+						if (bool bGeomPropHandled; !ShaderNode->GetBooleanAttribute(GeomPropHandledAttribute, bGeomPropHandled) || !bGeomPropHandled)
+						{
+							ShaderNode->AddBooleanAttribute(GeomPropHandledAttribute, true);
+							ShaderNodesTextureSampleUIDs.Emplace(ShaderNode->GetUniqueID());
+						}
+					}
+				}
+			}
+		);
+#endif	  // WITH_EDITOR
+
+		// Now we need to check if we have to set the slot of the mesh nodes here
+		if (TArray<FMaterialSlotMesh>* SlotMeshes = PrimPathToSlotMeshNodes.Find(MaterialUid))
+		{
+			if (FString* NewMaterialUID = MaterialUidToActualNodeUid.Find(MaterialUid))
+			{
+				for (const FMaterialSlotMesh& MaterialSlotMesh : *SlotMeshes)
+				{
+#if WITH_EDITOR
+					if (!GeomPropValues.IsEmpty())
+					{
+						MaterialSlotMesh.MeshNode->AddInt32Attribute(USD::Primvar::Number, GeomPropValues.Num());
+
+						for (int32 Index = 0; Index < GeomPropValues.Num(); ++Index)
+						{
+							MaterialSlotMesh.MeshNode->AddStringAttribute(USD::Primvar::Name + FString::FromInt(Index), GeomPropValues[Index].Name);
+							MaterialSlotMesh.MeshNode->AddBooleanAttribute(
+								USD::Primvar::TangentSpace + FString::FromInt(Index),
+								GeomPropValues[Index].bTangentSpace
+							);
+							MaterialSlotMesh.MeshNode->AddStringAttribute(
+								USD::Primvar::ShaderNodeTextureSample + FString::FromInt(Index),
+								ShaderNodesTextureSampleUIDs[Index]
+							);
+						}
+					}
+#endif	  // WITH_EDITOR
+					if (!MaterialSlotMesh.MeshNode->GetSlotMaterialDependencyUid(MaterialSlotMesh.MaterialSlotName, *NewMaterialUID))
+					{
+						MaterialSlotMesh.MeshNode->SetSlotMaterialDependencyUid(MaterialSlotMesh.MaterialSlotName, *NewMaterialUID);
+					}
+				}
+			}
+		}
+	};
+
+	bool UInterchangeUSDTranslatorImpl::AddMaterialXShaderGraph(
+		const UE::FUsdPrim& Prim,
+		UInterchangeUsdTranslatorSettings* TranslatorSettings,
+		UInterchangeBaseNodeContainer& NodeContainer,
+		TArray<FUsdMaterialXShaderGraph::FGeomProp>& GeomPropValues
+	)
+	{
+		FName RenderContext = TranslatorSettings ? TranslatorSettings->RenderContext : UnrealIdentifiers::UnrealRenderContext;
+
+		// Check for any references of MaterialX
+#if WITH_EDITOR
+		if (RenderContext == UnrealIdentifiers::MaterialXRenderContext)
+		{
+			TArray<FString> FilePaths = UsdUtils::GetMaterialXFilePaths(Prim);
+			for (const FString& File : FilePaths)
+			{
+				// the file has already been handled no need to do a Translate again
+				if (!Translators.Find(File))
+				{
+					UInterchangeManager& InterchangeManager = UInterchangeManager::GetInterchangeManager();
+					UInterchangeSourceData* SourceData = UInterchangeManager::CreateSourceData(File);
+
+					UInterchangeTranslatorBase* Translator = InterchangeManager.GetTranslatorForSourceData(SourceData);
+					// check on the Translator, it might return nullptr in case of reimport
+					if (Translator)
+					{
+						Translator->Translate(NodeContainer);
+						Translators.Emplace(File, Translator);
+					}
+				}
+			}
+
+			if (!FilePaths.IsEmpty())
+			{
+				return true;
+			}
+
+			// Only enable the shader graph on Windows for the time being as it crashes on Linux cause of a probable double free
+#if PLATFORM_WINDOWS
+			bool bHasShaderGraph = false;
+			FUsdMaterialXShaderGraph ShaderGraph(Prim, *RenderContext.ToString());
+			GeomPropValues = ShaderGraph.GetGeomPropValueNames();
+			if (MaterialX::DocumentPtr Document = ShaderGraph.GetDocument())
+			{
+				if (MaterialX::NodePtr SurfaceMaterial = Document->getNode(TCHAR_TO_UTF8(*Prim.GetName().ToString())))
+				{
+					try
+					{
+						TSharedPtr<FMaterialXBase> ShaderTranslator = FMaterialXManager::GetInstance().GetShaderTranslator(
+							SurfaceMaterial->getCategory().c_str(),
+							NodeContainer
+						);
+						bHasShaderGraph = ShaderTranslator.IsValid();
+						if (bHasShaderGraph)
+						{
+							ShaderTranslator->Translate(SurfaceMaterial);
+						}
+					}
+					catch (std::exception& Exception)
+					{
+						UE_LOG(LogUsd, Error, TEXT("%s"), UTF8_TO_TCHAR(Exception.what()));
+					}
+				}
+
+				return bHasShaderGraph;
+			}
+#endif	  // PLATFORM_WINDOWS
+		}
+#endif	  // WITH_EDITOR
+		return false;
+	}
+
+	// We use this visitor to set UsdToUnreal::FParameterValue TVariant values onto UInterchangeShaderGraphNode.
+	//
+	// For now we only set attributes meant to be parsed as material instance parameters.
+	// If/whenever we want to support generating full material shader graphs from USD, we likely don't want to just fill out inputs
+	// into a rigid material function structures based on the shading model like the GLTF translator does, as USD materials can have
+	// custom shader graphs themselves. We'd either need to truly generate arbitrary interchange shader graphs here to be useful, or
+	// to delegate this work to MaterialX somehow (c.f. MaterialX materials baked into USD shader graphs)
+	struct FMaterialInstanceParameterValueVisitor
+	{
+		FMaterialInstanceParameterValueVisitor(
+			const UE::FUsdPrim& InPrim,
+			UInterchangeBaseNodeContainer& InNodeContainer,
+			UInterchangeShaderGraphNode& InMaterialNode,
+			const TMap<FString, int32>& InPrimvarToUVIndex
+		)
+			: Prim(InPrim)
+			, NodeContainer(InNodeContainer)
+			, MaterialNode(InMaterialNode)
+			, PrimvarToUVIndex(InPrimvarToUVIndex)
+		{
+		}
+
+		static FString MakeInstanceParameterName(const FString& InParameterName)
+		{
+			return UE::Interchange::USD::AttributeIdentifierPrefix + InParameterName;
+		}
+
+		void EnableTextureForChannel(bool bEnable) const
+		{
+			using namespace UE::Interchange::USD;
+
+			MaterialNode.AddFloatAttribute(
+				MakeInstanceParameterName(UseTextureParameterPrefix + *BaseParameterName + UseTextureParameterSuffix),
+				bEnable ? 1.0f : 0.0f
+			);
+		}
+
+		void operator()(const float Value) const
+		{
+			MaterialNode.AddFloatAttribute(MakeInstanceParameterName(*BaseParameterName), Value);
+			EnableTextureForChannel(false);
+		}
+
+		void operator()(const FVector& Value) const
+		{
+			MaterialNode.AddLinearColorAttribute(MakeInstanceParameterName(*BaseParameterName), FLinearColor{Value});
+			EnableTextureForChannel(false);
+		}
+
+		void operator()(const UsdToUnreal::FTextureParameterValue& Value) const
+		{
+			// Emit texture node itself (this is the main place where this happens)
+			// Note that the node name isn't just the texture path, as we may have multiple material users of this texture
+			// with different settings, and so we need separate translated nodes for each material and parameter
+			const FString TextureUid = FString::Printf(TEXT("Texture:%s:%s"), *Prim.GetPrimPath().GetString(), **BaseParameterName);
+			AddTextureNode(Prim, TextureUid, Value, NodeContainer);
+
+			// Actual texture assignment
+			MaterialNode.AddStringAttribute(	//
+				MakeInstanceParameterName(FString::Printf(TEXT("%sTexture"), **BaseParameterName)),
+				TextureUid
+			);
+			EnableTextureForChannel(true);
+
+			// UV transform
+			FLinearColor ScaleAndTranslation = FLinearColor{
+				Value.UVScale.GetVector()[0],
+				Value.UVScale.GetVector()[1],
+				Value.UVTranslation[0],
+				Value.UVTranslation[1]
+			};
+			MaterialNode.AddLinearColorAttribute(
+				MakeInstanceParameterName(FString::Printf(TEXT("%sScaleTranslation"), **BaseParameterName)),
+				ScaleAndTranslation
+			);
+			MaterialNode.AddFloatAttribute(	   //
+				MakeInstanceParameterName(FString::Printf(TEXT("%sRotation"), **BaseParameterName)),
+				Value.UVRotation
+			);
+
+			// UV index
+			if (const int32* FoundIndex = PrimvarToUVIndex.Find(Value.Primvar))
+			{
+				MaterialNode.AddFloatAttribute(	   //
+					MakeInstanceParameterName(*BaseParameterName + UE::Interchange::USD::UVIndexParameterSuffix),
+					*FoundIndex
+				);
+			}
+			else
+			{
+				UE_LOG(
+					LogUsd,
+					Warning,
+					TEXT("Failed to find primvar '%s' when setting material parameter '%s' on material '%s'. Available primvars and UV "
+						 "indices: %s.%s"),
+					*Value.Primvar,
+					**BaseParameterName,
+					*Prim.GetPrimPath().GetString(),
+					*UsdUtils::StringifyMap(PrimvarToUVIndex),
+					Value.Primvar.IsEmpty()
+						? TEXT(
+							  " Is your UsdUVTexture Shader missing the 'inputs:st' attribute? (It specifies which UV set to sample the texture with)"
+						  )
+						: TEXT("")
+				);
+			}
+
+			// Component mask (which channel of the texture to use)
+			FLinearColor ComponentMask = FLinearColor::Black;
+			switch (Value.OutputIndex)
+			{
+				case 0:	   // RGB
+					ComponentMask = FLinearColor{1.f, 1.f, 1.f, 0.f};
+					break;
+				case 1:	   // R
+					ComponentMask = FLinearColor{1.f, 0.f, 0.f, 0.f};
+					break;
+				case 2:	   // G
+					ComponentMask = FLinearColor{0.f, 1.f, 0.f, 0.f};
+					break;
+				case 3:	   // B
+					ComponentMask = FLinearColor{0.f, 0.f, 1.f, 0.f};
+					break;
+				case 4:	   // A
+					ComponentMask = FLinearColor{0.f, 0.f, 0.f, 1.f};
+					break;
+			}
+			MaterialNode.AddLinearColorAttribute(
+				MakeInstanceParameterName(FString::Printf(TEXT("%sTextureComponent"), **BaseParameterName)),
+				ComponentMask
+			);
+		}
+
+		void operator()(const UsdToUnreal::FPrimvarReaderParameterValue& Value) const
+		{
+			MaterialNode.AddLinearColorAttribute(MakeInstanceParameterName(*BaseParameterName), FLinearColor{Value.FallbackValue});
+
+			if (Value.PrimvarName == TEXT("displayColor"))
+			{
+				MaterialNode.AddFloatAttribute(MakeInstanceParameterName(TEXT("UseVertexColorForBaseColor")), 1.0f);
+			}
+		}
+
+		void operator()(const bool Value) const
+		{
+			MaterialNode.AddBooleanAttribute(MakeInstanceParameterName(*BaseParameterName), static_cast<float>(Value));
+		}
+
+	public:
+		const UE::FUsdPrim& Prim;
+		UInterchangeBaseNodeContainer& NodeContainer;
+		UInterchangeShaderGraphNode& MaterialNode;
+		const TMap<FString, int32>& PrimvarToUVIndex;
+		const FString* BaseParameterName = nullptr;
+	};
+
+	TArray<FString> UInterchangeUSDTranslatorImpl::AddMaterialNode(
+		const UE::FUsdPrim& Prim,
+		UInterchangeUsdTranslatorSettings* TranslatorSettings,
+		UInterchangeBaseNodeContainer& NodeContainer
+	)
+	{
+		FName RenderContext = TranslatorSettings ? TranslatorSettings->RenderContext : UnrealIdentifiers::UniversalRenderContext;
+
+		// If this material has an unreal surface output and we're in the unreal render context, just emit a material reference,
+		// as we never want this to become a UMaterial / UMaterialInstance anyway.
+		//
+		// We could just early out here completely and not emit anything, as we also emit the material reference node on-demand, whenever
+		// we parse an actual material assignment from a Mesh. The user may have custom pipelines that expect to find these though,
+		// even if no mesh is actually using the materials
+		if (RenderContext == UnrealIdentifiers::UnrealRenderContext)
+		{
+			TOptional<FString> UnrealContentPath = UsdUtils::GetUnrealSurfaceOutput(Prim);
+			if (UnrealContentPath.IsSet())
+			{
+				return { UE::InterchangeUsdTranslator::Private::AddUnrealMaterialReferenceNodeIfNeeded(NodeContainer, UnrealContentPath.GetValue()) };
+			}
+		}
+
+		TArray<FString> MaterialUids;
+
+		FString PrimPath = Prim.GetPrimPath().GetString();
+		FString MaterialUid = MaterialPrefix + PrimPath;
+		FString MaterialTwoSidedUid = MaterialUid + TwoSidedSuffix;
+		FString MaterialPrimName(Prim.GetName().ToString());
+		bool bDuplicateMaterialToBeTwoSided = PrimPathToSlotMeshNodes.Contains(MaterialTwoSidedUid);
+
+		MaterialUids.Emplace(MaterialUid);
+		if (bDuplicateMaterialToBeTwoSided)
+		{
+			MaterialUids.Emplace(MaterialTwoSidedUid);
+		}
+
+		// Check if Node already exist as a TwoSided material
+		if (const UInterchangeShaderGraphNode* Node = Cast<UInterchangeShaderGraphNode>(NodeContainer.GetNode(MaterialTwoSidedUid)))
+		{
+			return MaterialUids;
+		}
+
+		// Check if Node already exist with this ID
+		if (const UInterchangeShaderGraphNode* ShaderGraphNode = Cast<UInterchangeShaderGraphNode>(NodeContainer.GetNode(MaterialUid)))
+		{
+			return MaterialUids;
+		}
+
+		// we only create material instances if we didn't find any MaterialX instances (including a shader graph)
+		if (TArray<FUsdMaterialXShaderGraph::FGeomProp> GeomPropValues;
+			AddMaterialXShaderGraph(Prim, TranslatorSettings, NodeContainer, GeomPropValues))
+		{
+			SetMaterialSlotDependencies(MaterialPrimName, MaterialUid, NodeContainer, GeomPropValues, true);
+			return MaterialUids;
+		}
+
+		UInterchangeShaderGraphNode* MaterialNode = NewObject<UInterchangeShaderGraphNode>(&NodeContainer);
+		MaterialNode->InitializeNode(MaterialUid, MaterialPrimName, EInterchangeNodeContainerType::TranslatedAsset);
+		MaterialNode->SetAssetName(MaterialPrimName);
+		MaterialNode->AddBooleanAttribute(UE::Interchange::USD::ParseMaterialIdentifier, true);
+		NodeContainer.AddNode(MaterialNode);
+
+		// Set the material instance node to the correct mesh nodes
+		MaterialUidToActualNodeUid.Add(MaterialUid, MaterialUid);
+		SetMaterialSlotDependencies(MaterialPrimName, MaterialUid, NodeContainer);
+
+		UsdToUnreal::FUsdPreviewSurfaceMaterialData MaterialData;
+		const bool bSuccess = UsdToUnreal::ConvertMaterial(Prim, MaterialData, *RenderContext.ToString());
+
+		// Set all the parameter values to the interchange node
+		bool bNeedsVTParent = false;
+		FMaterialInstanceParameterValueVisitor Visitor{Prim, NodeContainer, *MaterialNode, MaterialData.PrimvarToUVIndex};
+		for (TPair<FString, UsdToUnreal::FParameterValue>& Pair : MaterialData.Parameters)
+		{
+			Visitor.BaseParameterName = &Pair.Key;
+			Visit(Visitor, Pair.Value);
+
+			// Also simultaneously check if any of these parameters wants to be an UDIM texture so that we can use the VT reference material later
+			if (!bNeedsVTParent)
+			{
+				if (UsdToUnreal::FTextureParameterValue* TextureParameter = Pair.Value.TryGet<UsdToUnreal::FTextureParameterValue>())
+				{
+					if (TextureParameter->bIsUDIM)
+					{
+						bNeedsVTParent = true;
+					}
+				}
+			}
+		}
+
+		// Also set our parameter to uv index mapping as-is also as custom attributes, so that the USD Pipeline can make
+		// primvar-compatible materials
+		{
+			// Parameter to primvar
+			for (const TPair<FString, UsdToUnreal::FParameterValue>& Parameter : MaterialData.Parameters)
+			{
+				if (const UsdToUnreal::FTextureParameterValue* TextureParameterValue = Parameter.Value.TryGet<UsdToUnreal::FTextureParameterValue>())
+				{
+					const FString& MaterialParameter = Parameter.Key;
+					const FString& Primvar = TextureParameterValue->Primvar;
+
+					MaterialNode->AddStringAttribute(UE::Interchange::USD::ParameterToPrimvarAttributePrefix + MaterialParameter, Primvar);
+				}
+			}
+
+			// Primvar to UVIndex
+			for (const TPair<FString, int32>& Pair : MaterialData.PrimvarToUVIndex)
+			{
+				const FString& Primvar = Pair.Key;
+				int32 UVIndex = Pair.Value;
+
+				MaterialNode->AddInt32Attribute(UE::Interchange::USD::PrimvarUVIndexAttributePrefix + Primvar, UVIndex);
+			}
+		}
+
+		// Find and set the relevant info so that we can pick the reference material on the pipeline
+		//
+		// TODO: Proper VT texture support (we'd need to know the texture resolution at this point, and we haven't parsed them yet...).
+		// The way it currently works on Interchange is that the factory will create a VT or nonVT version of the texture to match the
+		// material parameter slot. Since we'll currently never set the VT reference material, it essentially means it will always
+		// downgrade our VT textures to non-VT.
+		// The only exception is how we upgrade the reference material to VT in case we have any UDIM textures a few lines above,
+		// as those are trivial to check for (we don't have to actually load the textures to do it)
+		if (UsdUtils::IsMaterialTranslucent(MaterialData))
+		{
+			MaterialNode->SetCustomBlendMode(EBlendMode::BLEND_Translucent);
+		}
+		if (bNeedsVTParent)
+		{
+			MaterialNode->AddBooleanAttribute(UE::Interchange::USD::NeedsVTParentIdentifier, true);
+		}
+
+		// Check if we have to duplicate the material to be TwoSided
+		if (bDuplicateMaterialToBeTwoSided)
+		{
+			FString MaterialAssetName = MaterialPrimName + TwoSidedSuffix;
+			UInterchangeShaderGraphNode* MaterialTwoSidedNode = NewObject<UInterchangeShaderGraphNode>(&NodeContainer);
+			UInterchangeBaseNode::CopyStorage(MaterialNode, MaterialTwoSidedNode);
+			MaterialTwoSidedNode->InitializeNode(MaterialTwoSidedUid, MaterialAssetName, EInterchangeNodeContainerType::TranslatedAsset);
+			MaterialTwoSidedNode->SetCustomTwoSided(true);
+			MaterialTwoSidedNode->SetCustomTwoSidedTransmission(true);
+			NodeContainer.AddNode(MaterialTwoSidedNode);
+
+			MaterialUidToActualNodeUid.Add(MaterialTwoSidedUid, MaterialTwoSidedUid);
+			SetMaterialSlotDependencies(MaterialAssetName, MaterialTwoSidedUid, NodeContainer);
+		}
+		return MaterialUids;
+	}
+
+	FString UInterchangeUSDTranslatorImpl::GetTextureSourcePath(const FString& TexturePathOnDisk)
+	{
+		if (!USDZFilePath.IsEmpty())
+		{
+			return USDZFilePath;
+		}
+		else
+		{
+			return TexturePathOnDisk;
+		}
+	}
+
+	void UInterchangeUSDTranslatorImpl::CleanUpDecompressedUSDZFolder()
+	{
+		if (!DecompressedUSDZRoot.IsEmpty())
+		{
+			const bool bRequireExists = false;
+			const bool bTree = true;
+			IFileManager::Get().DeleteDirectory(*FPaths::GetPath(DecompressedUSDZRoot), bRequireExists, bTree);
+		}
+
+		USDZFilePath.Reset();
+		DecompressedUSDZRoot.Reset();
+	}
+
+	void UInterchangeUSDTranslatorImpl::SetupTranslationContext(const UInterchangeUsdTranslatorSettings& Settings)
+	{
+		if (!TranslationContext)
+		{
+			TranslationContext = MakeShared<FUsdSchemaTranslationContext>(UsdStage);
+		}
+
+		TranslationContext->bIsImporting = true;
+		TranslationContext->Time = UsdUtils::GetDefaultTimeCode();
+		TranslationContext->bMergeIdenticalMaterialSlots = true;	// Interchange always does this
+		TranslationContext->bAllowInterpretingLODs = false;			// We don't support USD LODs yet
+
+		TranslationContext->PurposesToLoad = EUsdPurpose::Default | EUsdPurpose::Proxy | EUsdPurpose::Render | EUsdPurpose::Guide;
+		TranslationContext->RenderContext = Settings.RenderContext;
+		TranslationContext->MaterialPurpose = Settings.MaterialPurpose;
+	}
+
+	bool DecompressUSDZFileToTempFolder(const FString& InUSDZFilePath, FString& OutDecompressedUSDZRoot)
+	{
+		const bool bIncludeDot = false;
+		FString Extension = FPaths::GetExtension(InUSDZFilePath, bIncludeDot).ToLower();
+		if (Extension == TEXT("usdz"))
+		{
+			const FString Prefix = FPaths::GetBaseFilename(InUSDZFilePath);
+			const FString EmptyExtension = TEXT("");
+			const FString TempFolder = FPaths::CreateTempFilename(FPlatformProcess::UserTempDir(), *Prefix, *EmptyExtension);
+			FString DecompressedRoot;
+			const bool bSuccess = UsdUtils::DecompressUSDZFile(InUSDZFilePath, TempFolder, &DecompressedRoot);
+			if (bSuccess && !DecompressedRoot.IsEmpty())
+			{
+				OutDecompressedUSDZRoot = DecompressedRoot;
+				return true;
+			}
+			else
+			{
+				UE_LOG(LogUsd, Warning, TEXT("Failed to decompress USDZ file '%s': Textures may not be handled correctly."), *InUSDZFilePath);
+			}
+		}
+
+		return false;
+	}
+
+	UE::FUsdPrim UInterchangeUSDTranslatorImpl::TryGettingInactiveLODPrim(const FString& PrimPathString)
+	{
+		if (PrimPathString.IsEmpty())
+		{
+			return {};
+		}
+
+		UE::FSdfPath PrimPath{*PrimPathString};
+
+		const FString PrimName = PrimPath.GetName();
+		if (!IsValidLODName(PrimName))
+		{
+			return {};
+		}
+
+		const UE::FSdfPath LODContainerPath = PrimPath.GetParentPath();
+		const FString LODContainerPathString = LODContainerPath.GetString();
+		const FString& VariantName = PrimName;	  // Our convention for LODs is that the prim name matches the variant (e.g. "LOD2")
+
+		UE::FUsdStage TempStage;
+
+		{
+			FWriteScopeLock Lock{PrimPathToVariantToStageLock};
+
+			// Check if we have the stage we want already
+			if (TMap<FString, UE::FUsdStage>* TempStagesForPrim = PrimPathToVariantToStage.Find(LODContainerPathString))
+			{
+				if (UE::FUsdStage* TempStagesForVariant = TempStagesForPrim->Find(VariantName))
+				{
+					TempStage = *TempStagesForVariant;
+				}
+			}
+
+			// If not, let's open a new masked stage.
+			// USD will only compose the LOD container prim (and its subtree), so this should be very small and quick to open.
+			// We won't use the stage cache for these, so our strong reference right here is the only thing holding the stage opened.
+			if (!TempStage)
+			{
+				TempStage = UnrealUSDWrapper::OpenMaskedStage(
+					*UsdStage.GetRootLayer().GetIdentifier(),
+					EUsdInitialLoadSet::LoadAll,
+					{LODContainerPathString}
+				);
+
+				if (!TempStage)
+				{
+					return {};
+				}
+
+				UE::FUsdPrim LODContainerPrim = TempStage.GetPrimAtPath(LODContainerPath);
+				if (!LODContainerPrim)
+				{
+					return {};
+				}
+
+				// We have to edit the session layer here, and not the root layer directly. This because USD only opens a layer
+				// once in memory, so if we have multiple of these temp stages all trying to set the variant to different values
+				// on the root layer itself, they'd be actually trying to overwrite each other and could even lead to threading issues.
+				//
+				// The session layer however is unique to each of these temp stages so we won't have that problem, and it still
+				// should compose the variant switch just the same.
+				UE::FUsdEditContext Context{TempStage, TempStage.GetSessionLayer()};
+
+				UE::FUsdVariantSets VariantSets = LODContainerPrim.GetVariantSets();
+				bool bSwitched = VariantSets.SetSelection(LODString, VariantName);
+				if (!bSwitched)
+				{
+					return {};
+				}
+
+				// We finally have our stage for the particular lod container and with the variant we want: Cache it for later, if needed.
+				PrimPathToVariantToStage.FindOrAdd(LODContainerPathString).Add(VariantName, TempStage);
+			}
+		}
+
+		return TempStage.GetPrimAtPath(PrimPath);
+	}
+
+	void UInterchangeUSDTranslatorImpl::FinalizeLODContainerTraversal(
+		UInterchangeBaseNodeContainer& NodeContainer,
+		const FTraversalInfo& Info,
+		const UInterchangeSceneNode* SceneNodeWithLODs
+	)
+	{
+		if (!SceneNodeWithLODs)
+		{
+			return;
+		}
+
+		// Add a dedicated LODContainer node. We need this because every child of a LodGroup node
+		// will be interpreted as a LOD mesh, and we could have ended up with any additional children so far
+		// since any regular prim can contain the LOD variant (it could have any number of other children like
+		// other transforms, lights, skeleton prims, etc.)
+		const FString LODContainerUid = SceneNodeWithLODs->GetUniqueID() + LODContainerSuffix;
+		UInterchangeSceneNode* LODContainer = NewObject<UInterchangeSceneNode>(&NodeContainer);
+		{
+			LODContainer->InitializeNode(LODContainerUid, SceneNodeWithLODs->GetDisplayLabel(), EInterchangeNodeContainerType::TranslatedScene);
+			LODContainer->AddSpecializedType(UE::Interchange::FSceneNodeStaticData::GetLodGroupSpecializeTypeString());
+
+			NodeContainer.AddNode(LODContainer);
+			NodeContainer.SetNodeParentUid(LODContainerUid, SceneNodeWithLODs->GetUniqueID());
+		}
+
+		struct FLODSortHelper
+		{
+			UInterchangeSceneNode* SceneNode = nullptr;
+			int32 LODIndex;	   // Using actual ints in here instead of sorting on the node DisplayLabel directly so we correctly sort LOD10 > LOD2
+		};
+
+		TArray<FLODSortHelper> SortedLODs;
+		SortedLODs.Reserve(CurrentLODSceneNodes.Num());
+		for (UInterchangeSceneNode* Node : CurrentLODSceneNodes)
+		{
+			SortedLODs.Add({Node, GetLODIndexFromName(Node->GetDisplayLabel())});
+		}
+		SortedLODs.Sort(
+			[](const FLODSortHelper& LHS, const FLODSortHelper& RHS)
+			{
+				return LHS.LODIndex < RHS.LODIndex;
+			}
+		);
+
+		// Parent our LOD nodes to the LODContainer in the right order, because Interchange will try assigning
+		// those meshes to LOD numbers in the order it traverses the children, and there's no guarantee about the order we run
+		// into the variants when traversing the stage (or the order with which they were authored)
+		for (int32 Index = 0; Index < CurrentLODSceneNodes.Num(); ++Index)
+		{
+			const FLODSortHelper& PackedNode = SortedLODs[Index];
+			const FString& NodeUID = PackedNode.SceneNode->GetUniqueID();
+
+			NodeContainer.SetNodeParentUid(NodeUID, LODContainerUid);
+			NodeContainer.SetNodeDesiredChildIndex(NodeUID, Index);
+		}
+
+		// Detect invalid setup of LOD morph targets: Morph targets used by LOD2 must be present in LOD1, and all of those must be present in LOD0.
+		// On legacy USD support we used a hack to work around that, but it's not so trivial to do in Interchange and it's probably the wrong approach
+		// anyway: we should just let the user know how to correct their data for Unreal instead
+		bool bIsFirstLOD = true;
+		TSet<const UInterchangeMeshNode*> AllowedMorphTargets;
+		for (const FLODSortHelper& SortedLODNode : SortedLODs)
+		{
+			FString AssetInstanceUid;
+			if (!SortedLODNode.SceneNode->GetCustomAssetInstanceUid(AssetInstanceUid))
+			{
+				continue;
+			}
+
+			const UInterchangeMeshNode* MeshNode = Cast<const UInterchangeMeshNode>(NodeContainer.GetNode(AssetInstanceUid));
+			if (!MeshNode)
+			{
+				continue;
+			}
+
+			// Only skinned meshes have blend shapes / morph targets in USD
+			if (!MeshNode->IsSkinnedMesh())
+			{
+				continue;
+			}
+
+			TArray<FString> LODMorphTargetUids;
+			MeshNode->GetMorphTargetDependencies(LODMorphTargetUids);
+
+			TSet<const UInterchangeMeshNode*> LODMorphTargets;
+			LODMorphTargets.Reserve(LODMorphTargetUids.Num());
+
+			for (const FString& MorphTargetUid : LODMorphTargetUids)
+			{
+				if (const UInterchangeMeshNode* MorphTargetNode = Cast<const UInterchangeMeshNode>(NodeContainer.GetNode(MorphTargetUid)))
+				{
+					LODMorphTargets.Add(MorphTargetNode);
+				}
+			}
+
+			if (!bIsFirstLOD && !AllowedMorphTargets.Includes(LODMorphTargets))
+			{
+				UE_LOG(
+					LogUsd,
+					Warning,
+					TEXT(
+						"Invalid blend shape configuration for skeletal mesh LOD '%s': The set of blend shapes names used by lower LODs should include all blend shape names used by higher LODs"
+					),
+					*MeshNode->GetUniqueID()
+				);
+				break;
+			}
+
+			Swap(AllowedMorphTargets, LODMorphTargets);
+			bIsFirstLOD = false;
+		}
+
+		CurrentLODSceneNodes.Reset();
+	}
+
+	FString AddLightNode(const UE::FUsdPrim& Prim, UInterchangeBaseNodeContainer& NodeContainer)
+	{
+		FString NodeUid;
+
+		NodeUid = LightPrefix + Prim.GetPrimPath().GetString();
+		FString NodeName(Prim.GetName().ToString());
+
+		// Ref. UsdToUnreal::ConvertLight
+		static const FString IntensityToken = UsdToUnreal::ConvertToken(pxr::UsdLuxTokens->inputsIntensity);
+		static const FString ExposureToken = UsdToUnreal::ConvertToken(pxr::UsdLuxTokens->inputsExposure);
+		static const FString ColorToken = UsdToUnreal::ConvertToken(pxr::UsdLuxTokens->inputsColor);
+
+		float Intensity = UsdUtils::GetAttributeValue<float>(Prim, IntensityToken);
+		float Exposure = UsdUtils::GetAttributeValue<float>(Prim, ExposureToken);
+		FLinearColor Color = UsdUtils::GetAttributeValue<FLinearColor>(Prim, ColorToken);
+
+		const bool bSRGB = true;
+		Color = Color.ToFColor(bSRGB);
+
+		static const FString TemperatureToken = UsdToUnreal::ConvertToken(pxr::UsdLuxTokens->inputsColorTemperature);
+		static const FString UseTemperatureToken = UsdToUnreal::ConvertToken(pxr::UsdLuxTokens->inputsEnableColorTemperature);
+
+		float Temperature = UsdUtils::GetAttributeValue<float>(Prim, TemperatureToken);
+		bool UseTemperature = UsdUtils::GetAttributeValue<bool>(Prim, UseTemperatureToken);
+
+		// "Shadow enabled" currently not supported
+
+		auto SetBaseLightProperties = [&NodeUid, &NodeName, Color, Temperature, UseTemperature](UInterchangeBaseLightNode* LightNode)
+		{
+			LightNode->InitializeNode(NodeUid, NodeName, EInterchangeNodeContainerType::TranslatedAsset);
+			LightNode->SetAssetName(NodeName);
+
+			LightNode->SetCustomLightColor(Color);
+			LightNode->SetCustomTemperature(Temperature);
+			LightNode->SetCustomUseTemperature(UseTemperature);
+		};
+
+		static const FString RadiusToken = UsdToUnreal::ConvertToken(pxr::UsdLuxTokens->inputsRadius);
+
+		if (Prim.IsA(TEXT("DistantLight")))
+		{
+			UInterchangeDirectionalLightNode* LightNode = NewObject<UInterchangeDirectionalLightNode>(&NodeContainer);
+			SetBaseLightProperties(LightNode);
+
+			Intensity = UsdToUnreal::ConvertLightIntensityAttr(Intensity, Exposure);
+			LightNode->SetCustomIntensity(Intensity);
+
+			// LightSourceAngle currently not supported by UInterchangeDirectionalLightNode
+			// float Angle = UsdUtils::GetAttributeValue<float>(Prim, TEXTVIEW("inputs:angle"));
+
+			NodeContainer.AddNode(LightNode);
+		}
+		else if (Prim.IsA(TEXT("SphereLight")))
+		{
+			const FUsdStageInfo StageInfo(Prim.GetStage());
+
+			const float Radius = UsdUtils::GetAttributeValue<float>(Prim, RadiusToken);
+			const float SourceRadius = UsdToUnreal::ConvertDistance(StageInfo, Radius);	   // currently not supported
+
+			if (Prim.HasAPI(TEXT("ShapingAPI")))
+			{
+				UInterchangeSpotLightNode* LightNode = NewObject<UInterchangeSpotLightNode>(&NodeContainer);
+				SetBaseLightProperties(LightNode);
+
+				LightNode->SetCustomIntensityUnits(EInterchangeLightUnits::Lumens);
+
+				static const FString ConeAngleToken = UsdToUnreal::ConvertToken(pxr::UsdLuxTokens->inputsShapingConeAngle);
+				static const FString ConeSoftnessToken = UsdToUnreal::ConvertToken(pxr::UsdLuxTokens->inputsShapingConeSoftness);
+
+				float ConeAngle = UsdUtils::GetAttributeValue<float>(Prim, ConeAngleToken);
+				float ConeSoftness = UsdUtils::GetAttributeValue<float>(Prim, ConeSoftnessToken);
+
+				float InnerConeAngle = 0.0f;
+				const float OuterConeAngle = UsdToUnreal::ConvertConeAngleSoftnessAttr(ConeAngle, ConeSoftness, InnerConeAngle);
+
+				Intensity = UsdToUnreal::ConvertLuxShapingAPIIntensityAttr(Intensity, Exposure, Radius, ConeAngle, ConeSoftness, StageInfo);
+				LightNode->SetCustomIntensity(Intensity);
+
+				LightNode->SetCustomInnerConeAngle(InnerConeAngle);
+				LightNode->SetCustomOuterConeAngle(OuterConeAngle);
+
+				NodeContainer.AddNode(LightNode);
+			}
+			else
+			{
+				UInterchangePointLightNode* LightNode = NewObject<UInterchangePointLightNode>(&NodeContainer);
+				SetBaseLightProperties(LightNode);
+
+				LightNode->SetCustomIntensityUnits(EInterchangeLightUnits::Lumens);
+
+				Intensity = UsdToUnreal::ConvertSphereLightIntensityAttr(Intensity, Exposure, Radius, StageInfo);
+				LightNode->SetCustomIntensity(Intensity);
+
+				NodeContainer.AddNode(LightNode);
+			}
+		}
+		else if (Prim.IsA(TEXT("RectLight")) || Prim.IsA(TEXT("DiskLight")))
+		{
+			UInterchangeRectLightNode* LightNode = NewObject<UInterchangeRectLightNode>(&NodeContainer);
+			SetBaseLightProperties(LightNode);
+
+			LightNode->SetCustomIntensityUnits(EInterchangeLightUnits::Lumens);
+
+			static const FString WidthToken = UsdToUnreal::ConvertToken(pxr::UsdLuxTokens->inputsWidth);
+			static const FString HeightToken = UsdToUnreal::ConvertToken(pxr::UsdLuxTokens->inputsHeight);
+
+			float Width = UsdUtils::GetAttributeValue<float>(Prim, WidthToken);
+			float Height = UsdUtils::GetAttributeValue<float>(Prim, HeightToken);
+
+			const FUsdStageInfo StageInfo(Prim.GetStage());
+
+			if (Prim.IsA(TEXT("RectLight")))
+			{
+				Width = UsdToUnreal::ConvertDistance(StageInfo, Width);
+				Height = UsdToUnreal::ConvertDistance(StageInfo, Height);
+				Intensity = UsdToUnreal::ConvertRectLightIntensityAttr(Intensity, Exposure, Width, Height, StageInfo);
+			}
+			else
+			{
+				float Radius = UsdUtils::GetAttributeValue<float>(Prim, RadiusToken);
+				Width = UsdToUnreal::ConvertDistance(StageInfo, Radius) * 2.f;
+				Height = Width;
+
+				Intensity = UsdToUnreal::ConvertDiskLightIntensityAttr(Intensity, Exposure, Radius, StageInfo);
+			}
+			LightNode->SetCustomIntensity(Intensity);
+			LightNode->SetCustomSourceWidth(Width);
+			LightNode->SetCustomSourceHeight(Height);
+
+			NodeContainer.AddNode(LightNode);
+		}
+		// #ueent_todo:
+		// DomeLight -> SkyLight
+
+		return NodeUid;
+	}
+
+	FString AddCameraNode(const UE::FUsdPrim& Prim, UInterchangeBaseNodeContainer& NodeContainer)
+	{
+		FString NodeUid;
+
+		NodeUid = CameraPrefix + Prim.GetPrimPath().GetString();
+		FString NodeName(Prim.GetName().ToString());
+
+		UInterchangePhysicalCameraNode* CameraNode = NewObject<UInterchangePhysicalCameraNode>(&NodeContainer);
+		CameraNode->InitializeNode(NodeUid, NodeName, EInterchangeNodeContainerType::TranslatedAsset);
+		NodeContainer.AddNode(CameraNode);
+
+		// ref. UsdToUnreal::ConvertGeomCamera
+		UE::FUsdStage Stage = Prim.GetStage();
+		FUsdStageInfo StageInfo(Stage);
+
+		static const FString FocalLengthToken = UsdToUnreal::ConvertToken(pxr::UsdGeomTokens->focalLength);
+		static const FString HorizontalApertureToken = UsdToUnreal::ConvertToken(pxr::UsdGeomTokens->horizontalAperture);
+		static const FString VerticalApertureToken = UsdToUnreal::ConvertToken(pxr::UsdGeomTokens->verticalAperture);
+
+		float FocalLength = UsdUtils::GetAttributeValue<float>(Prim, FocalLengthToken);
+		FocalLength = UsdToUnreal::ConvertDistance(StageInfo, FocalLength);
+		CameraNode->SetCustomFocalLength(FocalLength);
+
+		float SensorWidth = UsdUtils::GetAttributeValue<float>(Prim, HorizontalApertureToken);
+		SensorWidth = UsdToUnreal::ConvertDistance(StageInfo, SensorWidth);
+		CameraNode->SetCustomSensorWidth(SensorWidth);
+
+		float SensorHeight = UsdUtils::GetAttributeValue<float>(Prim, VerticalApertureToken);
+		SensorHeight = UsdToUnreal::ConvertDistance(StageInfo, SensorHeight);
+		CameraNode->SetCustomSensorHeight(SensorHeight);
+
+		// Focus distance and FStop not currently supported
+
+		return NodeUid;
+	}
+
+	void AddMorphTargetNodes(
+		const UE::FUsdPrim& MeshPrim,
+		UInterchangeUSDTranslatorImpl& TranslatorImpl,
+		UInterchangeMeshNode& MeshNode,
+		UInterchangeBaseNodeContainer& NodeContainer,
+		const FTraversalInfo& Info
+	)
+	{
+		UE::FUsdSkelBlendShapeQuery Query{MeshPrim};
+		if (!Query)
+		{
+			return;
+		}
+
+		const FString MeshPrimPath = MeshPrim.GetPrimPath().GetString();
+
+		TFunction<void(const FString&, int32, const FString&)> AddMorphTargetNode =
+			[&MeshNode, &MeshPrimPath, &NodeContainer, &Info](const FString& MorphTargetName, int32 BlendShapeIndex, const FString& InbetweenName)
+		{
+			// Note: We identify a blend shape by its Mesh prim path and the blend shape index, even though
+			// the blend shape itself is a full standalone prim. This is for two reasons:
+			//  - We need to also read the Mesh prim's mesh data when emitting the payload, so having the Mesh path on the payload key is handy;
+			//  - It could be possible for different meshes to share the same BlendShape (possibly?), so we really want a separate version of
+			//    a blend shape for each mesh that uses it.
+			//
+			// Despite of that though, we won't use the blendshape's full path as the morph target name, so that users can get different
+			// blendshapes across the model to combine into a single morph target. Interchange has an import option to let you control
+			// whether they become separate morph targets or not anyway ("Merge Morph Targets with Same Name")
+			const FString NodeUid = GetMorphTargetMeshNodeUid(MeshPrimPath, BlendShapeIndex, InbetweenName);
+			const FString PayloadKey = GetMorphTargetMeshPayloadKey(Info.bInsideLOD, MeshPrimPath, BlendShapeIndex, InbetweenName);
+
+			UInterchangeMeshNode* MorphTargetMeshNode = NewObject<UInterchangeMeshNode>(&NodeContainer);
+			MorphTargetMeshNode->InitializeNode(NodeUid, MorphTargetName, EInterchangeNodeContainerType::TranslatedAsset);
+			MorphTargetMeshNode->SetPayLoadKey(PayloadKey, EInterchangeMeshPayLoadType::MORPHTARGET);
+			MorphTargetMeshNode->SetMorphTarget(true);
+			MorphTargetMeshNode->SetMorphTargetName(MorphTargetName);
+			NodeContainer.AddNode(MorphTargetMeshNode);
+			MeshNode.SetMorphTargetDependencyUid(NodeUid);
+		};
+
+		for (size_t Index = 0; Index < Query.GetNumBlendShapes(); ++Index)
+		{
+			UE::FUsdSkelBlendShape BlendShape = Query.GetBlendShape(Index);
+			if (!BlendShape)
+			{
+				continue;
+			}
+			UE::FUsdPrim BlendShapePrim = BlendShape.GetPrim();
+			const FString BlendShapeName = BlendShapePrim.GetName().ToString();
+
+			const FString UnusedInbetweenName;
+			AddMorphTargetNode(BlendShapeName, Index, UnusedInbetweenName);
+
+			for (const UE::FUsdSkelInbetweenShape& Inbetween : BlendShape.GetInbetweens())
+			{
+				const FString InbetweenName = Inbetween.GetAttr().GetName().ToString();
+				const FString MorphTargetName = BlendShapeName + TEXT("_") + InbetweenName;
+				AddMorphTargetNode(MorphTargetName, Index, InbetweenName);
+			}
+		}
+	}
+
+	FString UInterchangeUSDTranslatorImpl::AddMeshNode(
+		const UE::FUsdPrim& Prim,
+		UInterchangeBaseNodeContainer& NodeContainer,
+		const FTraversalInfo& Info,
+		bool bPrimitiveShape
+	)
+	{
+		FString PrimPath = Prim.GetPrimPath().GetString();
+		FString NodeUid = (bPrimitiveShape ? PrimitiveShapePrefix : MeshPrefix) + PrimPath;
+		FString NodeName(Prim.GetName().ToString());
+
+		// Check if Node already exist with this ID
+		if (const UInterchangeMeshNode* Node = Cast<UInterchangeMeshNode>(NodeContainer.GetNode(NodeUid)))
+		{
+			return NodeUid;
+		}
+
+		// Fill in the MeshNode itself
+		const bool bIsAnimated = UsdUtils::IsAnimatedMesh(Prim);
+		UInterchangeMeshNode* MeshNode = bIsAnimated ? NewObject<UInterchangeGeometryCacheNode>(&NodeContainer)
+													 : NewObject<UInterchangeMeshNode>(&NodeContainer);
+		MeshNode->InitializeNode(NodeUid, NodeName, EInterchangeNodeContainerType::TranslatedAsset);
+		MeshNode->SetAssetName(NodeName);
+		const bool bIsSkinned = static_cast<bool>(Info.ClosestParentSkelRootPath) && Prim.HasAPI(TEXT("SkelBindingAPI"));
+
+		FString PayloadKey = PrimPath;
+		if (Info.bInsideLOD)
+		{
+			PayloadKey = LODPrefix + PayloadKey;
+		}
+		if (bPrimitiveShape)
+		{
+			// We are currently not supporting Skinned Primitive Shapes.
+			// In theory a Skinned Mesh needs Joint influences and weights information provided,
+			// However, there does not seem to be a ruleset against a Primitive Shape having SkelBindingAPI set.
+			// Which means, that on a theoretical level there could be such a scenario.
+			ensureMsgf(!bIsSkinned, TEXT("Unexpected scenario: Primitive Shape is skinned."));
+
+			// For Primitive Shapes we add PrimitiveShapePrefix for the PayLoadKey,
+			// in order to be able to identify the PrimitiveShape in the PayloadData acquisition phase,
+			// as the Primitive Shapes require a different MeshDescription acquisition path, compared to Static Meshes.
+			PayloadKey = PrimitiveShapePrefix + PayloadKey;
+		}
+
+		double TimeCode = UsdUtils::GetDefaultTimeCode();
+		if (UInterchangeGeometryCacheNode* GeometryCacheNode = Cast<UInterchangeGeometryCacheNode>(MeshNode))
+		{
+			GeometryCacheNode->SetPayLoadKey(PrimPath, EInterchangeMeshPayLoadType::ANIMATED);
+
+			UE::FUsdStage Stage = Prim.GetStage();
+			int32 StartFrame = FMath::FloorToInt(Stage.GetStartTimeCode());
+			int32 EndFrame = FMath::CeilToInt(Stage.GetEndTimeCode());
+			UsdUtils::GetAnimatedMeshTimeCodes(Stage, PrimPath, StartFrame, EndFrame);
+
+			double FramesPerSecond = Stage.GetTimeCodesPerSecond();
+			if (FramesPerSecond == 0)
+			{
+				UE_LOG(LogUsd, Warning, TEXT("Invalid USD GeometryCache FPS detected. Falling back to 1 FPS"));
+				FramesPerSecond = 1;
+			}
+
+			// The GeometryCache module expects the end frame to be one past the last animation frame
+			EndFrame += 1;
+
+			GeometryCacheNode->SetCustomStartFrame(StartFrame);
+			GeometryCacheNode->SetCustomEndFrame(EndFrame);
+			GeometryCacheNode->SetCustomFrameRate(FramesPerSecond);
+
+			const bool bConstantTopology = UsdUtils::GetMeshTopologyVariance(Prim) != UsdUtils::EMeshTopologyVariance::Heterogenous;
+			GeometryCacheNode->SetCustomHasConstantTopology(bConstantTopology);
+
+			TimeCode = StartFrame;
+		}
+		else if (bIsSkinned && !bPrimitiveShape)
+		{
+			MeshNode->SetSkinnedMesh(true);
+			MeshNode->SetPayLoadKey(PayloadKey, EInterchangeMeshPayLoadType::SKELETAL);
+			if (Info.BoundSkeletonPrimPath && !Info.BoundSkeletonPrimPath->IsEmpty())
+			{
+				MeshNode->SetSkeletonDependencyUid(*Info.BoundSkeletonPrimPath);
+			}
+
+			AddMorphTargetNodes(Prim, *this, *MeshNode, NodeContainer, Info);
+
+			// When returning the payload data later, we'll need at the very least our SkeletonQuery, so
+			// here we store the Info object into the Impl
+			{
+				FWriteScopeLock ScopedInfoWriteLock{CachedTraversalInfoLock};
+				NodeUidToCachedTraversalInfo.Add(NodeUid, Info);
+			}
+		}
+		else
+		{
+			MeshNode->SetPayLoadKey(PayloadKey, EInterchangeMeshPayLoadType::STATIC);
+		}
+
+		// Material assignments
+		{
+			const bool bProvideMaterialIndices = false;
+			UsdUtils::FUsdPrimMaterialAssignmentInfo Assignments = UsdUtils::GetPrimMaterialAssignments(
+				Prim,
+				TimeCode,
+				bProvideMaterialIndices,
+				CachedMeshConversionOptions.RenderContext,
+				CachedMeshConversionOptions.MaterialPurpose
+			);
+
+			if (Info.bInsideLOD)
+			{
+				CachedMaterialAssignments.Add(PrimPath, Assignments);
+			}
+
+			// Move these into the asset node because the USD Pipeline will compare these with the assigned material's
+			// parameter-to-primvar mapping in order to make sure the mesh is using a primvar-compatible material.
+			//
+			// Note that ideally we'd cache this mapping and reuse it on the payload retrieval step. Instead, we will
+			// just end up calling the same function again during payload retrieval and hoping that it produces the same
+			// primvar-to-UVIndex mapping. It should though, as the mesh conversion options are the same. We can't cache
+			// the mapping because we run into USD allocator issues, given that all the strings contained in the
+			// FUsdPrimMaterialAssignmentInfo object are allocated inside an USD allocator scope
+			for (const TPair<FString, int32>& PrimvarPair : UsdUtils::GetPrimvarToUVIndexMap(Prim))
+			{
+				const FString& PrimvarName = PrimvarPair.Key;
+				int32 UVIndex = PrimvarPair.Value;
+
+				MeshNode->AddInt32Attribute(UE::Interchange::USD::PrimvarUVIndexAttributePrefix + PrimvarName, UVIndex);
+			}
+
+			for (const UsdUtils::FUsdPrimMaterialSlot& Slot : Assignments.Slots)
+			{
+				// We do this because Interchange will, in some scenarios, merge material slots with identical slot names.
+				// By using the source (which is the displaycolor desc / material prim path / unreal material content path)
+				// we do end up with a goofy looking super long material slot names, but it will have Interchange only combine
+				// slots if they really are pointing at the exact same thing
+				const FString& SlotName = Slot.MaterialSource;
+
+				// Get the Uid of the material instance that we'll end up assigning to this slot
+				FString MaterialInstanceUid;
+				bool bAddDirectlyAsSlotDependency = false;
+				switch (Slot.AssignmentType)
+				{
+					case UsdUtils::EPrimAssignmentType::DisplayColor:
+					{
+						bAddDirectlyAsSlotDependency = true;
+
+						// MaterialSource here is e.g. "!DisplayColor_0_1"
+						MaterialInstanceUid = AddDisplayColorMaterialInstanceNodeIfNeeded(NodeContainer, Slot);
+						break;
+					}
+					case UsdUtils::EPrimAssignmentType::MaterialPrim:
+					{
+						MaterialInstanceUid = MaterialPrefix + Slot.MaterialSource;	   // This is the prim path
+						if (Slot.bMeshIsDoubleSided)
+						{
+							MaterialInstanceUid += TwoSidedSuffix;
+						}
+						break;
+					}
+					case UsdUtils::EPrimAssignmentType::UnrealMaterial:
+					{
+						bAddDirectlyAsSlotDependency = true;
+
+						// MaterialSource here is the content path, e.g. "/Game/MyFolder/Red.Red"
+						MaterialInstanceUid = AddUnrealMaterialReferenceNodeIfNeeded(NodeContainer, Slot.MaterialSource);
+						break;
+					}
+					default:
+					{
+						ensure(false);
+						break;
+					}
+				}
+
+				// If we're a DisplayColor/Unreal material, we know everything we need right now
+				if (bAddDirectlyAsSlotDependency)
+				{
+					MeshNode->SetSlotMaterialDependencyUid(SlotName, *MaterialInstanceUid);
+				}
+				// If we found a match let's set the slot to the corresponding Material right away, as we already must have traversed this material
+				else if (FString* ActualMaterialInstanceUid = MaterialUidToActualNodeUid.Find(MaterialInstanceUid))
+				{
+					MeshNode->SetSlotMaterialDependencyUid(SlotName, *ActualMaterialInstanceUid);
+				}
+				// Otherwise, we need to wait until the material prim itself is translated, as we may need to defer to another translator
+				// (e.g. MaterialX) for the translation, which could generate an entirely different translated node we can't know about yet
+				else
+				{
+					// one material can be attached to several meshes
+					PrimPathToSlotMeshNodes.FindOrAdd(MaterialInstanceUid).Add({SlotName, MeshNode});
+				}
+			}
+		}
+
+		NodeContainer.AddNode(MeshNode);
+		return NodeUid;
+	}
+
+	void UInterchangeUSDTranslatorImpl::AddLODMeshNodes(	//
+		const UE::FUsdPrim& Prim,
+		UInterchangeBaseNodeContainer& NodeContainer,
+		UInterchangeSceneNode& ParentSceneNode,
+		UInterchangeUsdTranslatorSettings* TranslatorSettings,
+		FTraversalInfo Info
+	)
+	{
+		UE::FUsdVariantSets VariantSets = Prim.GetVariantSets();
+		const UE::FUsdVariantSet LODVariantSet = VariantSets.GetVariantSet(LODString);
+		if (!LODVariantSet.IsValid())
+		{
+			return;
+		}
+
+		const FString ActiveVariant = LODVariantSet.GetVariantSelection();
+		if (!IsValidLODName(ActiveVariant))
+		{
+			return;
+		}
+
+		Info.ParentNode = &ParentSceneNode;
+
+		bool bSwitchedFromInitialVariant = false;
+		for (const FString& VariantName : LODVariantSet.GetVariantNames())
+		{
+			// The active variant will be parsed via regular traversal
+			if (VariantName == ActiveVariant || !IsValidLODName(VariantName))
+			{
+				continue;
+			}
+
+			{
+				// For creating the scene nodes themselves we'll switch the active variant on the currently opened stage
+				// (still using the session layer to minimize impact to the actual layer). This mainly so that we can retrieve
+				// and fetch and cache the correct material bindings for the LOD meshes. Later on we'll use separate stages
+				// with population masks to read the LODs concurrently, and we won't be able to resolve material bindings
+				UE::FUsdEditContext Context{UsdStage, UsdStage.GetSessionLayer()};
+
+				bool bSwitchedVariant = VariantSets.SetSelection(LODString, VariantName);
+				if (!bSwitchedVariant)
+				{
+					continue;
+				}
+				bSwitchedFromInitialVariant = true;
+			}
+
+			if (UE::FUsdPrim LODMeshPrim = GetLODMesh(Prim, VariantName))
+			{
+				Traverse(LODMeshPrim, *this, NodeContainer, TranslatorSettings, Info);
+			}
+			else
+			{
+				UE_LOG(
+					LogTemp,
+					Warning,
+					TEXT(
+						"Failed to parse a LOD Mesh from variant '%s' of prim '%s'. For automatic parsing of LODs, make sure there is a single Mesh prim within the variant, named exactly as the variant itself (e.g. 'LOD0', 'LOD1', etc.)"
+					),
+					*VariantName,
+					*Prim.GetPrimPath().GetString()
+				);
+			}
+		}
+
+		// Put the active variant back to what it originally was
+		if (bSwitchedFromInitialVariant)
+		{
+			{
+				UE::FUsdEditContext Context{UsdStage, UsdStage.GetSessionLayer()};
+
+				const bool bRestoredSelection = VariantSets.SetSelection(LODString, ActiveVariant);
+				ensure(bRestoredSelection);
+			}
+
+			// Recompute our skel cache here if we have any as ancestor, because switching variants could have
+			// invalidated some of its internal state about its descendant prims, which we'll need to be OK
+			// when handling the payloads
+			Info.RepopulateSkelCache(UsdStage);
+		}
+	}
+
+	void AddVolumeCustomAttributesToNode(const UsdUtils::FVolumePrimInfo& VolumePrimInfo, UInterchangeVolumeNode* VolumeNode)
+	{
+		using namespace UE::Interchange::USD;
+
+		// Convert from the {'velocity': {'AttributesA.R': 'X', 'AttributesA.G': 'Y', 'AttributesA.B': 'Z'}} style of mapping
+		// from grid info into {"AttributesA.X": "velocity_0", "AttributesA.G": "velocity_1", "AttributesA.B": "velocity_2"}
+		// mapping into the VolumeNode custom attributes
+		for (const TPair<FString, TMap<FString, FString>>& GridPair : VolumePrimInfo.GridNameToChannelComponentMapping)
+		{
+			const FString& GridName = GridPair.Key;	   // "density", "temperature", etc.
+
+			const TMap<FString, FString>& AttributesChannelToGridChannel = GridPair.Value;
+			for (const TPair<FString, FString>& AssignmentPair : AttributesChannelToGridChannel)
+			{
+				const FString& AttributeChannel = AssignmentPair.Key;	 // "AttributesA.R", "AttributesB.B", etc.
+				const FString& GridChannel = AssignmentPair.Value;		 // "X", "Y", "Z", etc.
+
+				const static TMap<FString, FString> AttributeChannelToAttributeKey{
+					{TEXT("AttributesA.R"), SparseVolumeTexture::AttributesAChannelR},
+					{TEXT("AttributesA.G"), SparseVolumeTexture::AttributesAChannelG},
+					{TEXT("AttributesA.B"), SparseVolumeTexture::AttributesAChannelB},
+					{TEXT("AttributesA.A"), SparseVolumeTexture::AttributesAChannelA},
+					{TEXT("AttributesB.R"), SparseVolumeTexture::AttributesBChannelR},
+					{TEXT("AttributesB.G"), SparseVolumeTexture::AttributesBChannelG},
+					{TEXT("AttributesB.B"), SparseVolumeTexture::AttributesBChannelB},
+					{TEXT("AttributesB.A"), SparseVolumeTexture::AttributesBChannelA},
+				};
+
+				const static TMap<FString, FString> GridChannelToComponentIndex{
+					{TEXT("X"), TEXT("0")},
+					{TEXT("Y"), TEXT("1")},
+					{TEXT("Z"), TEXT("2")},
+					{TEXT("W"), TEXT("3")},
+					{TEXT("R"), TEXT("0")},
+					{TEXT("G"), TEXT("1")},
+					{TEXT("B"), TEXT("2")},
+					{TEXT("A"), TEXT("3")},
+				};
+
+				const FString* FoundAttributeKey = AttributeChannelToAttributeKey.Find(AttributeChannel);
+				if (!FoundAttributeKey)
+				{
+					UE_LOG(LogUsd, Warning, TEXT("Failing to parse unreal:SVT:mappedAttributeChannels value '%s'"), *AttributeChannel);
+					continue;
+				}
+
+				const FString* FoundComponentIndex = GridChannelToComponentIndex.Find(GridChannel);
+				if (!FoundComponentIndex)
+				{
+					UE_LOG(LogUsd, Warning, TEXT("Failing to parse unreal:SVT:mappedGridComponents value '%s'"), *GridChannel);
+					continue;
+				}
+
+				VolumeNode->AddStringAttribute(
+					*FoundAttributeKey,
+					GridName + UE::Interchange::Volume::GridNameAndComponentIndexSeparator + *FoundComponentIndex
+				);
+			}
+		}
+
+		// Convert texture format
+		static_assert((int)UsdUtils::ESparseVolumeAttributesFormat::Unorm8 == (int)EInterchangeSparseVolumeTextureFormat::Unorm8);
+		static_assert((int)UsdUtils::ESparseVolumeAttributesFormat::Float16 == (int)EInterchangeSparseVolumeTextureFormat::Float16);
+		static_assert((int)UsdUtils::ESparseVolumeAttributesFormat::Float32 == (int)EInterchangeSparseVolumeTextureFormat::Float32);
+		if (VolumePrimInfo.AttributesAFormat.IsSet())
+		{
+			VolumeNode->AddInt32Attribute(SparseVolumeTexture::AttributesAFormat, static_cast<int32>(VolumePrimInfo.AttributesAFormat.GetValue()));
+		}
+		if (VolumePrimInfo.AttributesBFormat.IsSet())
+		{
+			VolumeNode->AddInt32Attribute(SparseVolumeTexture::AttributesBFormat, static_cast<int32>(VolumePrimInfo.AttributesBFormat.GetValue()));
+		}
+	}
+
+	UInterchangeMaterialInstanceNode* CreateVolumetricMaterialInstanceNode(
+		const UE::FUsdPrim& VolumePrim,
+		UInterchangeBaseNodeContainer& NodeContainer,
+		const UsdToUnreal::FUsdMeshConversionOptions& ConversionOptions
+	)
+	{
+		FString ParentMaterialPath;
+
+		// Priority 1: Explicit material assignment on the Volume prim (Unreal materials).
+		{
+			double TimeCode = UsdUtils::GetDefaultTimeCode();	 // Not relevant as material bindings can't be time sampled
+			bool bProvideMaterialIndices = false;
+			UsdUtils::FUsdPrimMaterialAssignmentInfo Assignments = UsdUtils::GetPrimMaterialAssignments(
+				VolumePrim,
+				TimeCode,
+				bProvideMaterialIndices,
+				UnrealIdentifiers::Unreal,	  // Unreal materials should always be on the unreal render context
+				ConversionOptions.MaterialPurpose
+			);
+
+			for (const UsdUtils::FUsdPrimMaterialSlot& Slot : Assignments.Slots)
+			{
+				if (Slot.AssignmentType == UsdUtils::EPrimAssignmentType::UnrealMaterial)
+				{
+					ParentMaterialPath = Slot.MaterialSource;
+					break;
+				}
+			}
+		}
+
+		// Priority 2: The USD default volumetric material on the UsdProjectSettings.
+		bool bIsFallback = false;
+		if (ParentMaterialPath.IsEmpty())
+		{
+			if (const UUsdProjectSettings* ProjectSettings = GetDefault<UUsdProjectSettings>())
+			{
+				bIsFallback = true;
+				ParentMaterialPath = ProjectSettings->ReferenceDefaultSVTMaterial.ToString();
+			}
+		}
+
+		// Priority 3: Hard-coded fallback volumetric material that ships with the engine.
+		if (ParentMaterialPath.IsEmpty())
+		{
+			bIsFallback = true;
+			ParentMaterialPath = TEXT("/Engine/EngineMaterials/SparseVolumeMaterial.SparseVolumeMaterial");
+		}
+
+		FString MaterialDisplayLabel = bIsFallback ? TEXT("USDVolumetricFallbackMaterial") : FPaths::GetBaseFilename(ParentMaterialPath);
+		FString MaterialNodeUid = MaterialPrefix + MaterialDisplayLabel;
+		MakeNodeUidUniqueInContainer(MaterialNodeUid, NodeContainer);
+
+		// We'll always spawn a new material instance for each volume prim. Realistically a stage is only going to have a handful
+		// of volumes at most, and material instances should be pretty cheap. This should be a more predictable result for the user,
+		// and it prevents us from needing some bespoke code to reuse these material instance nodes depending on their volume assignment
+		UInterchangeMaterialInstanceNode* MaterialInstance = NewObject<UInterchangeMaterialInstanceNode>(&NodeContainer);
+		MaterialInstance->InitializeNode(MaterialNodeUid, MaterialDisplayLabel, EInterchangeNodeContainerType::TranslatedAsset);
+		MaterialInstance->SetCustomParent(ParentMaterialPath);
+
+		NodeContainer.AddNode(MaterialInstance);
+		return MaterialInstance;
+	}
+
+	TArray<FString> UInterchangeUSDTranslatorImpl::AddVolumeNodes(
+		const UE::FUsdPrim& InPrim,
+		UInterchangeBaseNodeContainer& InNodeContainer,
+		FString& OutMaterialInstanceUid,
+		bool& bOutNeedsFrameTrack
+	)
+	{
+		using namespace UE::Interchange::USD;
+
+		// Each Volume prim can reference multiple FieldAsset prims. Each FieldAsset itself can point to a particular grid within
+		// a .vdb file... USD is probably too flexible here, allowing us to reference grids from separate .vdb files in the same Volume prim,
+		// or letting us refer to the same grid more than once, etc.
+		//
+		// Our end goal is to make each .vdb file into a single SparseVolumeTexture, combining all the grids that need to be read from it.
+		// We'll do our best to satisfy all the requirements, but emit some warnings if we fail on an edge case. Then we want to spawn a single
+		// HeterogeneousVolumeActor for each Volume prim, generate an instance of the right material, and assign all these generated SVTs to it.
+		//
+		// In Interchange terms, this means that we'll emit a single InterchangeVolumeNode for each .vdb file, but will
+		// emit a new UInterchangeVolumeGridNode for each grid reference within that file. Note that USD's flexibility means we may have
+		// separate Volume prims all referencing the same shared FieldAsset prim, so we need to presume a InterchangeVolumeNode for this .vdb
+		// file has potentially already been created when parsing another Volume prim...
+		//
+		// Finally, since we may use the same .vdb file in multiple animations, and we want to end up with separate animated SVTs, we need
+		// separate factory nodes. If we want to keep the expected mapping of factory node / volume node Uids (just an added "Factory_" prefix)
+		// this means we need a separate volume node *per animation*, so we'll also use "AnimationIDs" to differentiate them
+
+		TStrongObjectPtr<UInterchangeBaseNodeContainer> TempContainer = nullptr;
+		TStrongObjectPtr<UInterchangeTranslatorBase> Translator = nullptr;
+
+		// This is collected by path hash here because for animated SVTs we want to still have a single UsdUtils::FVolumePrimInfo for each
+		// group of animated volume frames, since they will become a separate SVT. If we just collected them by filepath we could run into
+		// trouble if we had a volume prim with 3 frames starting at "file.vdb", and a separate volume prim that just wants one frame, "frame.vdb"
+		TMap<FString, UsdUtils::FVolumePrimInfo> VolumeInfoByFilePathHash = UsdUtils::GetVolumeInfoByFilePathHash(InPrim);
+
+		TSet<FString> VolumeAssetNodeUids;
+		TMap<FString, FString> VolumeFieldNameToNodeUids;
+
+		bOutNeedsFrameTrack = false;
+
+		// Stash these as we may need this info later when retrieving animation pipelines
+		TArray<UsdUtils::FVolumePrimInfo>& CollectedInfoForPrim = PrimPathToVolumeInfo.FindOrAdd(InPrim.GetPrimPath().GetString());
+
+		for (const TPair<FString, UsdUtils::FVolumePrimInfo>& Pair : VolumeInfoByFilePathHash)
+		{
+			const FString& AnimationID = Pair.Key;
+			const UsdUtils::FVolumePrimInfo& VolumePrimInfo = Pair.Value;
+			CollectedInfoForPrim.Add(VolumePrimInfo);
+
+			TArray<FString> VDBFilePaths;
+			VDBFilePaths.Reserve(VolumePrimInfo.TimeSamplePathIndices.Num() + 1);
+
+			// In case we have both timeSamples and different default opinion, add the default opinion as the first frame so that's what it shows
+			// on the level. The LevelSequence Frame track will factor this in, and have the LevelSequence only go through the TimeSamplePaths frames
+			// though.
+			bool bInsertedDefaultOpinion = false;
+			if (VolumePrimInfo.TimeSamplePaths.Num() > 0 && VolumePrimInfo.TimeSamplePaths[0] != VolumePrimInfo.SourceVDBFilePath)
+			{
+				VDBFilePaths.Add(VolumePrimInfo.SourceVDBFilePath);
+				bInsertedDefaultOpinion = true;
+			}
+			// No time samples at all
+			else if (VolumePrimInfo.TimeSamplePaths.Num() == 0)
+			{
+				VDBFilePaths.Add(VolumePrimInfo.SourceVDBFilePath);
+			}
+
+			// Add the file paths going through TimeSamplePathIndices because it's possible that GetVolumeInfoByFilePathHash
+			// deduplicated volume frames already. It's fine to add duplicate entries to VDBFilePaths though, because we'll check
+			// for an existing volume node for that path every time anyway
+			for (int32 PathIndex : VolumePrimInfo.TimeSamplePathIndices)
+			{
+				VDBFilePaths.Add(VolumePrimInfo.TimeSamplePaths[PathIndex]);
+			}
+
+			if (VolumePrimInfo.TimeSamplePathTimeCodes.Num() > 0)
+			{
+				bOutNeedsFrameTrack = true;
+			}
+
+			// First volume is special as that is what will "become" the animated factory node if we have animation. We'll also only stash
+			// our custom attributes on this first volume node. The filepaths are sorted according to timeSamples, so this is always the first
+			// frame of animation, or the default opinion if we have that
+			UInterchangeVolumeNode* FirstVolumeNode = nullptr;
+
+			for (int32 Index = 0; Index < VDBFilePaths.Num(); ++Index)
+			{
+				const FString& FilePath = VDBFilePaths[Index];
+
+				TMap<FString, UInterchangeVolumeNode*>& AnimationIDToVolumeNode = VolumeFilepathToAnimationIDToNode.FindOrAdd(FilePath);
+				UInterchangeVolumeNode* VolumeNode = AnimationIDToVolumeNode.FindRef(AnimationID);
+
+				// Need to translate the volume for this animation ID.
+				// Note: It may seem wasteful to translate the same volume more than once in case it is used by multiple animations,
+				// but keep in mind that:
+				//  - Multiple animations for the same volume frame in the same import will realistically never happen in practice;
+				//  - The VDB translator will cache the read file bytes from the first translation;
+				//  - "Translating" the volume just involves returning some header information, which should be pretty fast;
+				//  - Doing this saves us from having to manually duplicate other volume and grid nodes, and manually patching up
+				//    their unique IDs and / or making some sort of mistake;
+				if (!VolumeNode)
+				{
+					UInterchangeSourceData* SourceData = UInterchangeManager::CreateSourceData(FilePath);
+
+					if (!Translator || !Translator->CanImportSourceData(SourceData))
+					{
+						// Pass a USD context object to the translator, which is a signal that lets the OpenVDB translator
+						// be considered, even if its cvar is off (c.f. UInterchangeOpenVDBTranslator::CanImportSourceData)
+						UInterchangeUsdContext* Context = NewObject<UInterchangeUsdContext>();
+						SourceData->SetContextObjectByTag(UE::Interchange::USD::USDContextTag, Context);
+
+						UInterchangeManager& InterchangeManager = UInterchangeManager::GetInterchangeManager();
+						Translator.Reset(InterchangeManager.GetTranslatorForSourceData(SourceData));
+					}
+					if (!Translator)
+					{
+						const bool bIncludeDot = false;
+						const FString Extension = FPaths::GetExtension(FilePath, bIncludeDot).ToLower();
+						const bool bIsOpenVDB = Extension == TEXT("vdb");
+
+						UE_LOG(
+							LogUsd,
+							Error,
+							TEXT("Failed to find a compatible translator for file '%s'.%s"),
+							*FilePath,
+							bIsOpenVDB ? TEXT(" Is the 'Interchange OpenVDB Support' plugin enabled?") : TEXT("")
+						);
+						continue;
+					}
+
+					if (UInterchangeVolumeTranslatorSettings* Settings = Cast<UInterchangeVolumeTranslatorSettings>(Translator->GetSettings()))
+					{
+						// We never want to discover new .vdb files via the OpenVDB translator for animations.
+						// If we should have animations via USD they will be described on the USD file explicictly
+						Settings->bTranslateAdjacentNumberedFiles = false;
+
+						// If the volume prim describes an animation, let's add the same AnimationID to the volume nodes that the translator
+						// will output, so that the SVT pipeline groups them up into a single factory node
+						Settings->AnimationID = VolumePrimInfo.TimeSamplePathTimeCodes.Num() > 0 ? AnimationID : FString{};
+					}
+
+					if (!TempContainer)
+					{
+						TempContainer.Reset(NewObject<UInterchangeBaseNodeContainer>());
+					}
+					TempContainer->Reset();	   // Empty the container of nodes
+
+					Translator->SourceData = SourceData;
+					Translator->Translate(*TempContainer);
+
+					// Cache that we used this translator for this filepath. If we keep it, we don't have to open the file again to
+					// retrieve the payload. Note: This is likely the same translator we used for all .vdb files during this import
+					Translators.Add(FilePath, Translator);
+
+					// Move the nodes into our own node container. We do this because we must iterate to discover the volume node
+					// (as we have no guarantees about its NodeUid), and we don't want to do a full loop over our own NodeContainer
+					// that may have other Volume nodes and many more nodes overall.
+					//
+					// Also note the slight exploit: We can get non-const access to the translated nodes in this way, which we need
+					// in order to call InNodeContainer::Add
+					TempContainer->IterateNodes(
+						[&InNodeContainer, &VolumeNode](const FString& NodeUid, UInterchangeBaseNode* Node)
+						{
+							InNodeContainer.AddNode(Node);
+
+							if (UInterchangeVolumeNode* CastNode = Cast<UInterchangeVolumeNode>(Node))
+							{
+								if (VolumeNode == nullptr)
+								{
+									// Note that since we're not discovering new .vdbs, we expect exactly one UInterchangeVolumeNode to be generated
+									VolumeNode = CastNode;
+								}
+								else
+								{
+									UE_LOG(
+										LogUsd,
+										Warning,
+										TEXT("Found unexpected volume node '%s' (display label '%s')"),
+										*NodeUid,
+										*CastNode->GetDisplayLabel()
+									);
+								}
+							}
+						}
+					);
+
+					if (VolumeNode)
+					{
+						// Remove any animation index that the OpenVDB translator may have set on this volume node,
+						// as we want to control that explicitly from here
+						TArray<int32> ExistingAnimationIndices;
+						VolumeNode->GetCustomFrameIndicesInAnimation(ExistingAnimationIndices);
+						for (const int32 ExistingIndex : ExistingAnimationIndices)
+						{
+							VolumeNode->RemoveCustomFrameIndexInAnimation(ExistingIndex);
+						}
+
+						AnimationIDToVolumeNode.Add(AnimationID, VolumeNode);
+					}
+				}
+
+				if (!VolumeNode)
+				{
+					UE_LOG(LogUsd, Warning, TEXT("Failed to produce a volume node from file '%s'"), *FilePath);
+					continue;
+				}
+
+				VolumeNode->AddCustomFrameIndexInAnimation(Index);
+
+				if (!FirstVolumeNode)
+				{
+					FirstVolumeNode = VolumeNode;
+				}
+
+				VolumeAssetNodeUids.Add(VolumeNode->GetUniqueID());
+			}
+
+			if (FirstVolumeNode)
+			{
+				// Collect all of our custom-schema-based assignment info to be used as as custom attributes on the first volume node.
+				// The USD Pipeline will handle these, and move them into the factory nodes.
+				AddVolumeCustomAttributesToNode(VolumePrimInfo, FirstVolumeNode);
+
+				for (const FString& FieldName : VolumePrimInfo.VolumeFieldNames)
+				{
+					VolumeFieldNameToNodeUids.Add(FieldName, FirstVolumeNode->GetUniqueID());
+				}
+			}
+		}
+
+		// Setup a new material instance node for this volume, which will be used by the Heterogeneous Volume Actor we'll also
+		// spawn for this volume prim.
+		//
+		// Note that there's only one material slot per actor, but that we really do need some kind of material to be put in there,
+		// so that we can at least assign our SVTs somewhere.
+		UInterchangeMaterialInstanceNode* MaterialInstance = CreateVolumetricMaterialInstanceNode(
+			InPrim,
+			InNodeContainer,
+			CachedMeshConversionOptions
+		);
+
+		// Assign SVTs as material parameters on our new material instance
+		if (ensure(MaterialInstance))
+		{
+			OutMaterialInstanceUid = MaterialInstance->GetUniqueID();
+
+			TMultiMap<FString, FString> MaterialParameterToFieldName = UsdUtils::GetVolumeMaterialParameterToFieldNameMap(InPrim);
+
+			// Prim doesn't have the attributes specifying an explicit material parameter name to volume mapping (this is probably
+			// the more common case)
+			if (MaterialParameterToFieldName.Num() == 0)
+			{
+				// Consider that the field names may match material parameter names. We can't tell if that's the case or not from here,
+				// and this may cause us to assign a VolumeUid more than once (as multiple fields may target the same SVT), but the
+				// USDPipeline will clean this up
+				for (const TPair<FString, FString>& Pair : VolumeFieldNameToNodeUids)
+				{
+					const FString& FieldName = Pair.Key;
+					const FString& VolumeUid = Pair.Value;
+
+					MaterialInstance->AddTextureParameterValue(UE::Interchange::USD::VolumeFieldNameMaterialParameterPrefix + FieldName, VolumeUid);
+				}
+			}
+			// Prim has custom attributes specifying exactly which volume should be assigned to which material parameter
+			else
+			{
+				TMap<FString, FString> ParameterNameToVolume;
+				ParameterNameToVolume.Reserve(MaterialParameterToFieldName.Num());
+
+				for (const TPair<FString, FString>& Pair : MaterialParameterToFieldName)
+				{
+					const FString& ParameterName = Pair.Key;
+					const FString& FieldName = Pair.Value;
+
+					const FString* FoundVolumeUid = VolumeFieldNameToNodeUids.Find(FieldName);
+					if (!FoundVolumeUid)
+					{
+						continue;
+					}
+
+					// Show a warning in case we have conflicting assignments, because the legacy schema translator also did it
+					if (FString* FoundAssignedVolume = ParameterNameToVolume.Find(ParameterName))
+					{
+						if (*FoundAssignedVolume != *FoundVolumeUid)
+							UE_LOG(
+								LogUsd,
+								Warning,
+								TEXT(
+									"Trying to assign different Sparse Volume Textures to the same material parameter '%s' on the material instantiated for Volume '%s' and field name '%s'! Only a single texture can be assigned to a material parameter at a time."
+								),
+								*ParameterName,
+								**FoundVolumeUid,
+								*FieldName
+							);
+						continue;
+					}
+					ParameterNameToVolume.Add(ParameterName, *FoundVolumeUid);
+
+					MaterialInstance->AddTextureParameterValue(ParameterName, *FoundVolumeUid);
+				}
+			}
+		}
+
+		return VolumeAssetNodeUids.Array();
+	}
+
+	void AddTrackSetNode(UInterchangeUSDTranslatorImpl& Impl, UInterchangeBaseNodeContainer& NodeContainer)
+	{
+		// For now we only want a single track set (i.e. LevelSequence) per stage.
+		// TODO: One track set per layer, and add the tracks to the tracksets that correspond to layers where the opinions came from
+		// (similar to LevelSequenceHelper). Then we can use UInterchangeAnimationTrackSetInstanceNode to create "subsequences"
+		if (Impl.CurrentTrackSet)
+		{
+			return;
+		}
+
+		UE::FSdfLayer Layer = Impl.UsdStage.GetRootLayer();
+		const FString AnimTrackSetNodeUid = AnimationPrefix + Layer.GetIdentifier();
+		const FString AnimTrackSetNodeDisplayName = FPaths::GetBaseFilename(Layer.GetDisplayName());	// Strip extension
+
+		// We should only have one track set node per scene for now
+		const UInterchangeAnimationTrackSetNode* ExistingNode = Cast<UInterchangeAnimationTrackSetNode>(NodeContainer.GetNode(AnimTrackSetNodeUid));
+		if (!ensure(ExistingNode == nullptr))
+		{
+			return;
+		};
+
+		UInterchangeAnimationTrackSetNode* TrackSetNode = NewObject<UInterchangeAnimationTrackSetNode>(&NodeContainer);
+		TrackSetNode->InitializeNode(AnimTrackSetNodeUid, AnimTrackSetNodeDisplayName, EInterchangeNodeContainerType::TranslatedAsset);
+		TrackSetNode->SetCustomFrameRate(Layer.GetFramesPerSecond());	 // Key values in Interchange seem to be in seconds, so timeCodesPerSecond is
+																		 // not relevant here
+
+		NodeContainer.AddNode(TrackSetNode);
+		Impl.CurrentTrackSet = TrackSetNode;
+	}
+
+	void AddTransformAnimationNode(const UE::FUsdPrim& Prim, UInterchangeUSDTranslatorImpl& Impl, UInterchangeBaseNodeContainer& NodeContainer)
+	{
+		const FString PrimPath = Prim.GetPrimPath().GetString();
+		const FString UniquePath = PrimPath + TEXT("\\") + UnrealIdentifiers::TransformPropertyName.ToString();
+		const FString AnimTrackNodeUid = AnimationTrackPrefix + UniquePath;
+
+		const UInterchangeTransformAnimationTrackNode* ExistingNode = Cast<UInterchangeTransformAnimationTrackNode>(
+			NodeContainer.GetNode(AnimTrackNodeUid)
+		);
+		if (ExistingNode)
+		{
+			return;
+		}
+
+		UInterchangeTransformAnimationTrackNode* TransformAnimTrackNode = NewObject<UInterchangeTransformAnimationTrackNode>(&NodeContainer);
+		TransformAnimTrackNode->InitializeNode(AnimTrackNodeUid, UniquePath, EInterchangeNodeContainerType::TranslatedAsset);
+		TransformAnimTrackNode->SetCustomActorDependencyUid(*PrimPath);
+		TransformAnimTrackNode->SetCustomAnimationPayloadKey(UniquePath, EInterchangeAnimationPayLoadType::CURVE);
+		TransformAnimTrackNode->SetCustomUsedChannels((int32)EMovieSceneTransformChannel::AllTransform);
+
+		NodeContainer.AddNode(TransformAnimTrackNode);
+
+		AddTrackSetNode(Impl, NodeContainer);
+		Impl.CurrentTrackSet->AddCustomAnimationTrackUid(AnimTrackNodeUid);
+	}
+
+	void AddPropertyAnimationNode(
+		const FString& SceneNodeUid,
+		const FString& UEPropertyName,
+		EInterchangePropertyTracks TrackType,
+		EInterchangeAnimationPayLoadType PayloadType,
+		UInterchangeUSDTranslatorImpl& Impl,
+		UInterchangeBaseNodeContainer& NodeContainer
+	)
+	{
+		// We don't use the USD attribute path here because we want one unique node per UE track name,
+		// so that if e.g. both "intensity" and "exposure" are animated we make a single track for
+		// the Intensity UE property
+		const FString UniquePath = SceneNodeUid + TEXT("\\") + UEPropertyName;
+		const FString AnimTrackNodeUid = AnimationTrackPrefix + UniquePath;
+
+		const UInterchangeAnimationTrackNode* ExistingNode = Cast<UInterchangeAnimationTrackNode>(NodeContainer.GetNode(AnimTrackNodeUid));
+		if (ExistingNode)
+		{
+			return;
+		}
+
+		UInterchangeAnimationTrackNode* AnimTrackNode = NewObject<UInterchangeAnimationTrackNode>(&NodeContainer);
+		AnimTrackNode->InitializeNode(AnimTrackNodeUid, UniquePath, EInterchangeNodeContainerType::TranslatedAsset);
+		AnimTrackNode->SetCustomActorDependencyUid(*SceneNodeUid);
+		AnimTrackNode->SetCustomPropertyTrack(TrackType);
+		AnimTrackNode->SetCustomAnimationPayloadKey(UniquePath, PayloadType);
+
+		NodeContainer.AddNode(AnimTrackNode);
+
+		AddTrackSetNode(Impl, NodeContainer);
+		Impl.CurrentTrackSet->AddCustomAnimationTrackUid(AnimTrackNodeUid);
+	}
+
+	void AddPropertyAnimationNodes(const UE::FUsdPrim& Prim, UInterchangeUSDTranslatorImpl& Impl, UInterchangeBaseNodeContainer& NodeContainer)
+	{
+		using namespace UE::InterchangeUsdTranslator::Private;
+
+		if (!Prim)
+		{
+			return;
+		}
+		const FString PrimPath = Prim.GetPrimPath().GetString();
+
+		for (UE::FUsdAttribute Attr : Prim.GetAttributes())
+		{
+			if (!Attr || !Attr.ValueMightBeTimeVarying() || Attr.GetNumTimeSamples() == 0)
+			{
+				continue;
+			}
+
+			// Emit a STEPCURVE in case of a bool track: CURVE is only for floats/doubles (c.f. FLevelSequenceHelper::PopulateAnimationTrack).
+			// For now we're lucky in that all possible results from GetPropertiesForAttribute() are either all not bool, or either all bool,
+			// so we can reuse this for all the different UEAttrNames we get from the same attribute
+			const FName AttrTypeName = Attr.GetTypeName();
+			const bool bIsBoolTrack = AttrTypeName == TEXT("bool") || AttrTypeName == TEXT("token");	// Visibility is a token track
+
+			TArray<FName> UEPropertyNames = UsdUtils::GetPropertiesForAttribute(Prim, Attr.GetName().ToString());
+			for (const FName& UEPropertyName : UEPropertyNames)
+			{
+				const EInterchangePropertyTracks* FoundTrackType = PropertyNameToTrackType.Find(UEPropertyName);
+				if (!FoundTrackType)
+				{
+					continue;
+				}
+
+				const EInterchangeAnimationPayLoadType PayloadType = bIsBoolTrack ? EInterchangeAnimationPayLoadType::STEPCURVE
+																				  : EInterchangeAnimationPayLoadType::CURVE;
+
+				AddPropertyAnimationNode(PrimPath, UEPropertyName.ToString(), *FoundTrackType, PayloadType, Impl, NodeContainer);
+			}
+		}
+	}
+
+	// Some of the volume prim info is meant for the HeterogeneousVolume (HV) actor and the volumetric material,
+	// so we need to add it to the scene node (it's possible separate HV actors with different values for these
+	// end up sharing identical volume nodes)
+	void AddVolumeSceneNodeAttributes(
+		const UE::FUsdPrim& Prim,
+		UInterchangeSceneNode* SceneNode,
+		TArray<FString> AssetNodeUids,
+		const FString& VolumeMaterialInstanceUid,
+		bool bNeedsAnimationTrack,
+		UInterchangeUSDTranslatorImpl& Impl,
+		UInterchangeBaseNodeContainer& NodeContainer
+	)
+	{
+		if (!SceneNode)
+		{
+			return;
+		}
+		const FString& SceneNodeUid = SceneNode->GetUniqueID();
+
+		// Target the scene node from all the volume nodes to make it easy to find everything we need on the USD Pipeline
+		for (const FString& AssetNodeUid : AssetNodeUids)
+		{
+			if (const UInterchangeVolumeNode* VolumeNode = Cast<const UInterchangeVolumeNode>(NodeContainer.GetNode(AssetNodeUid)))
+			{
+				VolumeNode->AddTargetNodeUid(SceneNodeUid);
+			}
+		}
+
+		// Set our volumetric material as a "material override" directly on the scene node, which the USD Pipeline will also use
+		if (!VolumeMaterialInstanceUid.IsEmpty())
+		{
+			SceneNode->SetSlotMaterialDependencyUid(UE::Interchange::Volume::VolumetricMaterial, VolumeMaterialInstanceUid);
+		}
+
+		if (bNeedsAnimationTrack)
+		{
+			// Ideally we'd write some step curves, but Interchange doesn't support float step curves
+			const EInterchangeAnimationPayLoadType PayloadType = EInterchangeAnimationPayLoadType::CURVE;
+			const EInterchangePropertyTracks TrackType = EInterchangePropertyTracks::HeterogeneousVolumeFrame;
+			const static FName UEPropertyName = GET_MEMBER_NAME_CHECKED(UHeterogeneousVolumeComponent, Frame);
+
+			AddPropertyAnimationNode(SceneNodeUid, UEPropertyName.ToString(), TrackType, PayloadType, Impl, NodeContainer);
+		}
+	}
+
+	void AddSkeletalAnimationNode(
+		const UE::FUsdSkelSkeletonQuery& SkeletonQuery,
+		const TMap<FString, TPair<FString, int32>>& BoneToUidAndBoneIndex,
+		UInterchangeUSDTranslatorImpl& TranslatorImpl,
+		UInterchangeSceneNode& SkeletonPrimNode,
+		UInterchangeBaseNodeContainer& NodeContainer,
+		const FTraversalInfo& Info
+	)
+	{
+		UE::FUsdSkelAnimQuery AnimQuery = SkeletonQuery.GetAnimQuery();
+		if (!AnimQuery)
+		{
+			return;
+		}
+
+		UE::FUsdPrim SkelAnimationPrim = AnimQuery.GetPrim();
+		if (!SkelAnimationPrim)
+		{
+			return;
+		}
+
+		UE::FUsdPrim SkeletonPrim = SkeletonQuery.GetSkeleton();
+		if (!SkeletonPrim)
+		{
+			return;
+		}
+
+		UE::FUsdStage Stage = SkeletonPrim.GetStage();
+
+		const FString SkelAnimationName = SkelAnimationPrim.GetName().ToString();
+		const FString SkelAnimationPrimPath = SkelAnimationPrim.GetPrimPath().GetString();
+		const FString SkeletonPrimPath = SkeletonPrim.GetPrimPath().GetString();
+		const FString UniquePath = SkelAnimationPrimPath + TEXT("\\") + SkeletonPrimPath;
+		const FString NodeUid = AnimationTrackPrefix + UniquePath;
+
+		const UInterchangeSkeletalAnimationTrackNode* ExistingNode = Cast<UInterchangeSkeletalAnimationTrackNode>(NodeContainer.GetNode(NodeUid));
+		if (ExistingNode)
+		{
+			return;
+		};
+
+		UInterchangeSkeletalAnimationTrackNode* SkelAnimNode = NewObject<UInterchangeSkeletalAnimationTrackNode>(&NodeContainer);
+		SkelAnimNode->InitializeNode(NodeUid, SkelAnimationName, EInterchangeNodeContainerType::TranslatedAsset);
+		SkelAnimNode->SetCustomSkeletonNodeUid(SkeletonPrimNode.GetUniqueID());
+
+		// TODO: Uncomment this whenever Interchange supports skeletal animation sections, because
+		// currently it seems that InterchangeLevelSequenceFactory.cpp doesn't even have the string "skel" anywhere.
+		// If we were to add this all we'd get is a warning on the output log about "all referenced actors being missing",
+		// in case it failed to find anything else (e.g. other actual property/transform track) to put on the LevelSequence.
+		// AddTrackSetNode(TranslatorImpl, NodeContainer);
+		// TranslatorImpl->CurrentTrackSet->AddCustomAnimationTrackUid(NodeUid);
+
+		NodeContainer.AddNode(SkelAnimNode);
+
+		// Time info
+		{
+			// TODO: Match the TrackSet framerate whenever Interchange supports skeletal animation sections.
+			// float TrackSetFrameRate = 30.0f;
+			// if (TranslatorImpl->CurrentTrackSet->GetCustomFrameRate(TrackSetFrameRate))
+			// {
+			// 	SkelAnimNode->SetCustomAnimationSampleRate(TrackSetFrameRate);
+			// }
+			SkelAnimNode->SetCustomAnimationSampleRate(Stage.GetFramesPerSecond());
+
+			TOptional<double> StartTimeCode;
+			TOptional<double> StopTimeCode;
+
+			// For now we don't generate LevelSequences for sublayers and will instead put everything on a single
+			// LevelSequence for the entire stage, so we don't need to care so much about sublayer offset/scale like
+			// UsdToUnreal::ConvertSkelAnim does
+			TArray<double> JointTimeSamples;
+			if (AnimQuery.GetJointTransformTimeSamples(JointTimeSamples) && JointTimeSamples.Num() > 0)
+			{
+				StartTimeCode = JointTimeSamples[0];
+				StopTimeCode = JointTimeSamples[JointTimeSamples.Num() - 1];
+			}
+			TArray<double> BlendShapeTimeSamples;
+			if (AnimQuery.GetBlendShapeWeightTimeSamples(BlendShapeTimeSamples) && BlendShapeTimeSamples.Num() > 0)
+			{
+				StartTimeCode = FMath::Min(BlendShapeTimeSamples[0], StartTimeCode.Get(TNumericLimits<double>::Max()));
+				StopTimeCode = FMath::Max(BlendShapeTimeSamples[BlendShapeTimeSamples.Num() - 1], StopTimeCode.Get(TNumericLimits<double>::Lowest()));
+			}
+
+			UE::FUsdStage UsdStage = SkeletonPrim.GetStage();
+			double TimeCodesPerSecond = UsdStage.GetTimeCodesPerSecond();
+			if (StartTimeCode.IsSet())
+			{
+				SkelAnimNode->SetCustomAnimationStartTime(StartTimeCode.GetValue() / TimeCodesPerSecond);
+			}
+			if (StopTimeCode.IsSet())
+			{
+				SkelAnimNode->SetCustomAnimationStopTime(StopTimeCode.GetValue() / TimeCodesPerSecond);
+			}
+		}
+
+		// Joint animation
+		TArray<FString> UsdJointOrder = AnimQuery.GetJointOrder();
+		for (const FString& FullAnimatedBoneName : UsdJointOrder)
+		{
+			const TPair<FString, int32>* FoundPair = BoneToUidAndBoneIndex.Find(FullAnimatedBoneName);
+			if (!FoundPair)
+			{
+				continue;
+			}
+
+			const FString& BoneSceneNodeUid = FoundPair->Key;
+			int32 SkeletonOrderBoneIndex = FoundPair->Value;
+
+			const FString BoneAnimPayloadKey = SkeletonPrimPath + TEXT("\\") + LexToString(SkeletonOrderBoneIndex);
+
+			// When retrieving the payload later, We'll need that bone's index within the Skeleton prim to index into the
+			// InUsdSkeletonQuery.ComputeJointLocalTransforms() results.
+			// Note that we're describing joint transforms with baked frames here. It would have been possible to use transform
+			// curves, but that may have lead to issues when interpolating problematic joint transforms. Instead, we'll bake
+			// using USD, and let it interpolate the transforms however it wants
+			SkelAnimNode->SetAnimationPayloadKeyForSceneNodeUid(BoneSceneNodeUid, BoneAnimPayloadKey, EInterchangeAnimationPayLoadType::BAKED);
+		}
+
+		// Morph targets
+		{
+			UE::FUsdSkelBinding SkelBinding;
+			const bool bTraverseInstanceProxies = true;
+			bool bSuccess = Info.FurthestSkelCache->ComputeSkelBinding(	   //
+				Info.ResolveClosestParentSkelRoot(Stage),
+				SkeletonPrim,
+				SkelBinding,
+				bTraverseInstanceProxies
+			);
+			if (!bSuccess)
+			{
+				return;
+			}
+
+			TArray<FString> SkelAnimChannelOrder = AnimQuery.GetBlendShapeOrder();
+
+			TMap<FString, int32> SkelAnimChannelIndices;
+			SkelAnimChannelIndices.Reserve(SkelAnimChannelOrder.Num());
+			for (int32 ChannelIndex = 0; ChannelIndex < SkelAnimChannelOrder.Num(); ++ChannelIndex)
+			{
+				const FString& ChannelName = SkelAnimChannelOrder[ChannelIndex];
+				SkelAnimChannelIndices.Add(ChannelName, ChannelIndex);
+			}
+
+			TArray<UE::FUsdSkelSkinningQuery> SkinningTargets = SkelBinding.GetSkinningTargets();
+			for (const UE::FUsdSkelSkinningQuery& SkinningTarget : SkinningTargets)
+			{
+				// USD lets you "skin" anything that can take the SkelBindingAPI, but we only care about Mesh here as
+				// those are the only ones that can have blendshapes
+				UE::FUsdPrim Prim = SkinningTarget.GetPrim();
+				if (!Prim.IsA(TEXT("Mesh")))
+				{
+					continue;
+				}
+				const FString MeshPrimPath = Prim.GetPrimPath().GetString();
+
+				TArray<FString> BlendShapeChannels;
+				bool bInnerSucces = SkinningTarget.GetBlendShapeOrder(BlendShapeChannels);
+				if (!bInnerSucces)
+				{
+					continue;
+				}
+
+				TArray<UE::FSdfPath> Targets;
+				{
+					UE::FUsdRelationship BlendShapeTargetsRel = SkinningTarget.GetBlendShapeTargetsRel();
+					if (!BlendShapeTargetsRel)
+					{
+						continue;
+					}
+					bInnerSucces = BlendShapeTargetsRel.GetTargets(Targets);
+					if (!bInnerSucces)
+					{
+						continue;
+					}
+				}
+
+				if (BlendShapeChannels.Num() != Targets.Num())
+				{
+					UE_LOG(
+						LogUsd,
+						Warning,
+						TEXT(
+							"Skipping morph target curves for animation of skinned mesh '%s' because the number of entries in the 'skel:blendShapes' attribute (%d) doesn't match the number of entries in the 'skel:blendShapeTargets' attribute (%d)"
+						),
+						*MeshPrimPath,
+						BlendShapeChannels.Num(),
+						Targets.Num()
+					);
+					continue;
+				}
+
+				for (int32 BlendShapeIndex = 0; BlendShapeIndex < Targets.Num(); ++BlendShapeIndex)
+				{
+					const FString& ChannelName = BlendShapeChannels[BlendShapeIndex];
+					int32* FoundSkelAnimChannelIndex = SkelAnimChannelIndices.Find(ChannelName);
+					if (!FoundSkelAnimChannelIndex)
+					{
+						// This channel is not animated by this SkelAnimation prim
+						continue;
+					}
+
+					// Note that we put no inbetween name on the MorphTargetUid: We only need to emit the morph target curve payloads
+					// for the main shapes: We'll provide the inbetween "positions" when providing the curve and Interchange computes
+					// the inbetween curves automatically
+					const FString BlendShapePath = Targets[BlendShapeIndex].GetString();
+					const FString MorphTargetUid = GetMorphTargetMeshNodeUid(MeshPrimPath, BlendShapeIndex);
+					const FString PayloadKey = GetMorphTargetCurvePayloadKey(SkeletonPrimPath, *FoundSkelAnimChannelIndex, BlendShapePath);
+
+					SkelAnimNode->SetAnimationPayloadKeyForMorphTargetNodeUid(	  //
+						MorphTargetUid,
+						PayloadKey,
+						EInterchangeAnimationPayLoadType::MORPHTARGETCURVE
+					);
+				}
+			}
+		}
+	}
+
+	void AddSkeletonNodes(
+		const UE::FUsdPrim& Prim,
+		UInterchangeUSDTranslatorImpl& TranslatorImpl,
+		UInterchangeSceneNode& SkeletonPrimNode,
+		UInterchangeBaseNodeContainer& NodeContainer,
+		FTraversalInfo& Info
+	)
+	{
+		// If we're not inside of a SkelRoot, the skeleton shouldn't really do anything
+		if (!Info.ClosestParentSkelRootPath.IsValid())
+		{
+			return;
+		}
+
+		// By the time we get here we've already emitted a scene node for the skeleton prim itself, so we just
+		// need to emit a node hierarchy that mirrors the joints.
+
+		// Make the prim node into an Interchange joint/bone itself. By doing this we solve three issues:
+		//  - It becomes easy to identify our SkeletonDependencyUid when parsing Mesh nodes: It's just the skeleton prim path
+		//    (as opposed to having to target the translated node of the first root joint of the skeleton);
+		//  - We automatically handle USD skeletons with multiple root bones: We'll only ever have one "true"
+		//    root bone anyway: The SkeletonPrimNode itself;
+		//  - If a skeleton has no bones at all somehow, we'll still make one "bone" for it (this node).
+		SkeletonPrimNode.AddSpecializedType(UE::Interchange::FSceneNodeStaticData::GetJointSpecializeTypeString());
+		SkeletonPrimNode.SetCustomBindPoseLocalTransform(&NodeContainer, FTransform::Identity);
+		SkeletonPrimNode.SetCustomTimeZeroLocalTransform(&NodeContainer, FTransform::Identity);
+		const FString SkeletonPrimNodeUid = SkeletonPrimNode.GetUniqueID();
+
+#if WITH_EDITOR
+		// Convert the skeleton bones/joints into ConvertedData
+		UE::FUsdSkelSkeletonQuery SkelQuery = Info.FurthestSkelCache->GetSkelQuery(Prim);
+		const bool bEnsureAtLeastOneBone = false;
+		const bool bEnsureSingleRootBone = false;
+		UsdToUnreal::FUsdSkeletonData ConvertedData;
+		const bool bSuccess = UsdToUnreal::ConvertSkeleton(SkelQuery, ConvertedData, bEnsureAtLeastOneBone, bEnsureSingleRootBone);
+		if (!bSuccess)
+		{
+			return;
+		}
+
+		// Maps from the USD-style full bone name (e.g. "shoulder/elbow/hand") to the Uid we used for
+		// the corresponding scene node, and the bone's index on the skeleton's joint order.
+		// We'll need this to parse skeletal animations, if any
+		TMap<FString, TPair<FString, int32>> BoneToUidAndBoneIndex;
+
+		// Recursively traverse ConvertedData spawning the joint translated nodes
+		TFunction<void(int32, UInterchangeSceneNode&, const FString&)> RecursiveTraverseBones = nullptr;
+		RecursiveTraverseBones = [&RecursiveTraverseBones,
+								  &BoneToUidAndBoneIndex,
+								  &SkeletonPrimNodeUid,
+								  &ConvertedData,
+								  &NodeContainer	//
+		](int32 BoneIndex, UInterchangeSceneNode& ParentNode, const FString& BonePath)
+		{
+			const UsdToUnreal::FUsdSkeletonData::FBone& Bone = ConvertedData.Bones[BoneIndex];
+
+			// Reconcatenate a full "bone path" here for uniqueness, because Bone.Name is just the name of this
+			// single bone/joint itself (e.g. "Elbow")
+			const FString ConcatBonePath = (BonePath.IsEmpty() ? TEXT("") : (BonePath + TEXT("/"))) + Bone.Name;
+
+			// Putting the BonePrefix here avoids the pathological case where the user has skeleton child prims
+			// with names that match the joint names
+			const FString BoneNodeUid = SkeletonPrimNodeUid + BonePrefix + ConcatBonePath;
+
+			UInterchangeSceneNode* BoneNode = NewObject<UInterchangeSceneNode>(&NodeContainer);
+			BoneNode->InitializeNode(BoneNodeUid, Bone.Name, EInterchangeNodeContainerType::TranslatedScene);
+			BoneNode->AddSpecializedType(UE::Interchange::FSceneNodeStaticData::GetJointSpecializeTypeString());
+
+			// Note that we use our rest transforms for the Interchange bind pose as well: This because Interchange
+			// will put this on the RefSkeleton and so it will make its way to the Skeleton asset. We already kind
+			// of bake in our skeleton bind pose directly into our skinned mesh, so we really just want to put the
+			// rest pose on the skeleton asset/ReferenceSkeleton
+			BoneNode->SetCustomBindPoseLocalTransform(&NodeContainer, Bone.LocalRestTransform);
+			BoneNode->SetCustomTimeZeroLocalTransform(&NodeContainer, Bone.LocalRestTransform);
+			BoneNode->SetCustomLocalTransform(&NodeContainer, Bone.LocalRestTransform);
+
+			NodeContainer.AddNode(BoneNode);
+			NodeContainer.SetNodeParentUid(BoneNodeUid, ParentNode.GetUniqueID());
+
+			BoneToUidAndBoneIndex.Add(ConcatBonePath, {BoneNodeUid, BoneIndex});
+
+			for (int32 ChildIndex : Bone.ChildIndices)
+			{
+				RecursiveTraverseBones(ChildIndex, *BoneNode, ConcatBonePath);
+			}
+		};
+
+		// Start traversing from the root bones (we may have more than one, so check them all)
+		TSet<FString> UsedBoneNames;
+		for (int32 BoneIndex = 0; BoneIndex < ConvertedData.Bones.Num(); ++BoneIndex)
+		{
+			const UsdToUnreal::FUsdSkeletonData::FBone& Bone = ConvertedData.Bones[BoneIndex];
+			UsedBoneNames.Add(Bone.Name);
+
+			if (Bone.ParentIndex == INDEX_NONE)
+			{
+				const FString BonePathRoot = TEXT("");
+				RecursiveTraverseBones(BoneIndex, SkeletonPrimNode, BonePathRoot);
+			}
+		}
+
+		// Interchange will abort parsing skeletons that don't have unique names for each bone. If the user has that
+		// on their actual skeleton, then that's just invalid data and we can just let it fail and emit the error message.
+		// However, we don't want to end up with duplicate bone names and fail to parse when the duplicate "bone" is due to
+		// how we actually use the Skeleton prim itself as the root, as that's our little "trick". In this case, here we
+		// just change the display text of the skeleton prim itself to be unique (which is used for the bone name)
+		const FString SkeletonPrimName = SkeletonPrimNode.GetDisplayLabel();
+		FString NewSkeletonPrimName = UsdUnreal::ObjectUtils::GetUniqueName(SkeletonPrimName, UsedBoneNames);
+		if (NewSkeletonPrimName != SkeletonPrimName)
+		{
+			SkeletonPrimNode.SetDisplayLabel(NewSkeletonPrimName);
+		}
+
+		// Handle SkelAnimation prims, if we have any bound for this Skeleton
+		AddSkeletalAnimationNode(SkelQuery, BoneToUidAndBoneIndex, TranslatorImpl, SkeletonPrimNode, NodeContainer, Info);
+
+		// Cache our joint names in order, as this is needed when generating skeletal mesh payloads
+		Info.SkelJointNames = MakeShared<TArray<FString>>();
+		Info.SkelJointNames->Reserve(ConvertedData.Bones.Num());
+		for (const UsdToUnreal::FUsdSkeletonData::FBone& Bone : ConvertedData.Bones)
+		{
+			Info.SkelJointNames->Add(Bone.Name);
+		}
+		{
+			FWriteScopeLock ScopedInfoWriteLock{TranslatorImpl.CachedTraversalInfoLock};
+			TranslatorImpl.NodeUidToCachedTraversalInfo.Add(SkeletonPrimNodeUid, Info);
+		}
+#endif	  // WITH_EDITOR
+	}
+
+	template<typename T>
+	void TranslateAttribute(const UE::FUsdAttribute& Attr, const FString& AttrName, UInterchangeSceneNode* Node)
+	{
+		T Value;
+		if (Attr.Get<T>(Value))	   // Always check for an opinion on the default time code
+		{
+			TOptional<FString> PayloadKey;
+			UInterchangeUserDefinedAttributesAPI::CreateUserDefinedAttribute(Node, AttrName, Value, PayloadKey);
+		}
+	};
+
+	void TranslateAttributes(const UE::FUsdPrim& Prim, UInterchangeSceneNode* Node, const FString& AllowedAttributeRegex)
+	{
+		FRegexPattern RegexPattern{AllowedAttributeRegex};
+
+		for (const UE::FUsdAttribute& Attr : Prim.GetAttributes())
+		{
+			if (!Attr.HasAuthoredValue())
+			{
+				continue;
+			}
+
+			const FString AttrName = Attr.GetName().ToString();
+
+			FRegexMatcher RegexMatcher{RegexPattern, AttrName};
+			if (!RegexMatcher.FindNext())
+			{
+				continue;
+			}
+
+			// Note: We could have used Cpp type names here instead of the value type token, but doing the
+			// latter lets us handle the color types as FLinearColor, which is probably more useful
+			// References:
+			// - https://openusd.org/docs/api/_usd__page__datatypes.html
+			// - EAttributeTypes declaration
+			// - Engine/Source/Runtime/Interchange/Core/Private/Tests/StorageTest.cpp
+			// clang-format off
+			using TranslationFunc = TFunction<void(const UE::FUsdAttribute&, const FString&, UInterchangeSceneNode*)>;
+			static TMap<FName, TranslationFunc> TranslationFuncs = {
+				{TEXT("bool"),  		TranslateAttribute<bool>},
+
+				{TEXT("uchar"), 		TranslateAttribute<uint8>},
+				{TEXT("uchar[]"), 		TranslateAttribute<TArray<uint8>>},
+				{TEXT("int"), 			TranslateAttribute<int32>},
+				{TEXT("uint"), 			TranslateAttribute<uint32>},
+				{TEXT("int64"), 		TranslateAttribute<int64>},
+				{TEXT("uint64"), 		TranslateAttribute<uint64>},
+
+				{TEXT("half"),	 		TranslateAttribute<FFloat16>},
+				{TEXT("float"), 		TranslateAttribute<float>},
+				{TEXT("double"), 		TranslateAttribute<double>},
+				{TEXT("timecode"), 		TranslateAttribute<double>},
+
+				{TEXT("string"), 		TranslateAttribute<FString>},
+				{TEXT("token"), 		TranslateAttribute<FString>},
+				{TEXT("asset"), 		TranslateAttribute<FString>},
+
+				{TEXT("matrix2d"), 		TranslateAttribute<FMatrix44d>}, // Interchange only supports 4x4 matrices
+				{TEXT("matrix3d"), 		TranslateAttribute<FMatrix44d>},
+				{TEXT("matrix4d"), 		TranslateAttribute<FMatrix44d>},
+				{TEXT("frame4d"), 		TranslateAttribute<FMatrix44d>},
+
+				{TEXT("quath"), 		TranslateAttribute<FQuat4f>}, // No analogue type for Quat of halfs
+				{TEXT("quatf"), 		TranslateAttribute<FQuat4f>},
+				{TEXT("quatd"), 		TranslateAttribute<FQuat4d>},
+
+				{TEXT("half2"), 		TranslateAttribute<FVector2DHalf>},
+				{TEXT("float2"), 		TranslateAttribute<FVector2f>},
+				{TEXT("double2"), 		TranslateAttribute<FVector2d>},
+				{TEXT("int2"), 			TranslateAttribute<FIntPoint>},
+
+				{TEXT("half3"), 		TranslateAttribute<FVector3f>}, // There is no FVector3DHalf
+				{TEXT("point3h"), 		TranslateAttribute<FVector3f>},
+				{TEXT("normal3h"), 		TranslateAttribute<FVector3f>},
+				{TEXT("vector3h"), 		TranslateAttribute<FVector3f>},
+				{TEXT("color3h"), 		TranslateAttribute<FLinearColor>},
+
+				{TEXT("float3"), 		TranslateAttribute<FVector3f>},
+				{TEXT("point3f"), 		TranslateAttribute<FVector3f>},
+				{TEXT("normal3f"), 		TranslateAttribute<FVector3f>},
+				{TEXT("vector3f"), 		TranslateAttribute<FVector3f>},
+				{TEXT("color3f"), 		TranslateAttribute<FLinearColor>},
+
+				{TEXT("double3"), 		TranslateAttribute<FVector3d>},
+				{TEXT("point3d"), 		TranslateAttribute<FVector3d>},
+				{TEXT("normal3d"), 		TranslateAttribute<FVector3d>},
+				{TEXT("vector3d"), 		TranslateAttribute<FVector3d>},
+				{TEXT("color3d"), 		TranslateAttribute<FLinearColor>},
+
+				{TEXT("int3"), 			TranslateAttribute<FIntVector>},
+
+				{TEXT("half4"), 		TranslateAttribute<FVector4f>}, // There is no FVector4DHalf
+				{TEXT("float4"), 		TranslateAttribute<FVector4f>},
+				{TEXT("double4"), 		TranslateAttribute<FVector4d>},
+				{TEXT("int4"), 			TranslateAttribute<FIntRect>},
+				{TEXT("color4h"), 		TranslateAttribute<FLinearColor>},
+				{TEXT("color4f"), 		TranslateAttribute<FLinearColor>},
+				{TEXT("color4d"), 		TranslateAttribute<FLinearColor>},
+			};
+			// clang-format on
+
+			if (TranslationFunc Func = TranslationFuncs.FindRef(Attr.GetTypeName()))
+			{
+				Func(Attr, AttrName, Node);
+			}
+		}
+	}
+
+	void Traverse(
+		const UE::FUsdPrim& Prim,
+		UInterchangeUSDTranslatorImpl& TranslatorImpl,
+		UInterchangeBaseNodeContainer& NodeContainer,
+		UInterchangeUsdTranslatorSettings* TranslatorSettings,
+		FTraversalInfo Info
+	)
+	{
+		FString SceneNodeUid = Prim.GetPrimPath().GetString();
+		FString DisplayLabel(Prim.GetName().ToString());
+
+		// Do this before generating other nodes as they may need the updated info
+		Info.UpdateWithCurrentPrim(Prim);
+
+		FString VolumeMaterialInstanceUid;
+		bool bNeedsVolumeTrack = false;
+
+		// Generate asset node(s) if applicable
+		TArray<FString> AssetNodeUids;
+		if (Prim.IsA(TEXT("Material")))
+		{
+			AssetNodeUids = TranslatorImpl.AddMaterialNode(Prim, TranslatorSettings, NodeContainer);
+		}
+		else if (Prim.IsA(TEXT("Mesh")))
+		{
+			AssetNodeUids = {TranslatorImpl.AddMeshNode(Prim, NodeContainer, Info)};
+		}
+		else if (Prim.IsA(TEXT("Camera")))
+		{
+			AssetNodeUids = {AddCameraNode(Prim, NodeContainer)};
+		}
+		else if (Prim.HasAPI(TEXT("LightAPI")))
+		{
+			AssetNodeUids = {AddLightNode(Prim, NodeContainer)};
+		}
+		else if (Prim.IsA(TEXT("Gprim")) && !Prim.IsA(TEXT("PointBased")) && !Prim.IsA(TEXT("Volume")))
+		{
+			// PointBased prims are currently not supported apart from Meshes (which is taken care of in a previous branch).
+			// Volumes are also not currently supported.
+
+			constexpr bool bPrimitiveShape_True = true;
+			AssetNodeUids = {TranslatorImpl.AddMeshNode(Prim, NodeContainer, Info, bPrimitiveShape_True)};
+		}
+		else if (Prim.IsA(TEXT("Volume")))
+		{
+			AssetNodeUids = TranslatorImpl.AddVolumeNodes(Prim, NodeContainer, VolumeMaterialInstanceUid, bNeedsVolumeTrack);
+		}
+
+		// Only prims that require rendering (and have a renderable parent) get a scene node.
+		// This includes Xforms but also Scopes, which are not Xformable.
+		//
+		// We add a scene node for the pseudoroot in order to make tree traversal easier on the pipeline,
+		// but the pipeline will strip the pseudoroot node as its final step (if desired)
+		const bool bIsImageable = Prim.IsA(TEXT("Imageable"));
+		const bool bNeedsSceneNode = Prim.IsPseudoRoot()														  //
+									 || (bIsImageable && (Info.ParentNode || Prim.GetParent().IsPseudoRoot()))	  //
+									 || Info.bIsLODContainer;	 // We allow typeless prims (so not imageable) to be LOD containers
+
+		UInterchangeSceneNode* SceneNode = nullptr;
+		if (bNeedsSceneNode)
+		{
+			SceneNode = NewObject<UInterchangeSceneNode>(&NodeContainer);
+			SceneNode->InitializeNode(SceneNodeUid, DisplayLabel, EInterchangeNodeContainerType::TranslatedScene);
+			NodeContainer.AddNode(SceneNode);
+
+			// Store our purpose here to be filtered on the pipeline
+			const EUsdPurpose PrimPurpose = IUsdPrim::GetPurpose(Prim);
+			if (PrimPurpose != EUsdPurpose::Default)
+			{
+				SceneNode->AddInt32Attribute(UE::Interchange::USD::GeometryPurposeIdentifier, (int32)PrimPurpose);
+			}
+
+			// Store our prim kind as well, if we have any (this becomes the empty string if the prim has no authored kind)
+			const FString KindString = UsdToUnreal::ConvertToken(IUsdPrim::GetKind(Prim));
+			if (!KindString.IsEmpty())
+			{
+				const FString UserDefinedAttributeName = TEXT("kind");
+				TOptional<FString> PayloadKey;
+				UInterchangeUserDefinedAttributesAPI::CreateUserDefinedAttribute(SceneNode, UserDefinedAttributeName, KindString, PayloadKey);
+			}
+
+			if (TranslatorSettings->bTranslatePrimAttributes)
+			{
+				TranslateAttributes(Prim, SceneNode, TranslatorSettings->AttributeRegexFilter);
+			}
+
+			// If we're an Xformable, get our transform.
+			// All SceneNodes should have their LocalTransform set though.
+			// Not setting will cause ensure hits in Skeleton generations for example.
+			FTransform Transform = FTransform::Identity;
+			bool bResetTransformStack = false;
+			UsdToUnreal::ConvertXformable(Prim.GetStage(), UE::FUsdTyped(Prim), Transform, UsdUtils::GetEarliestTimeCode(), &bResetTransformStack);
+			SceneNode->SetCustomLocalTransform(&NodeContainer, Transform);
+
+			// Hide our scene node if it is meant to be hidden in USD
+			//
+			// We use actor visibility so that it matches how we map visibility timeSamples to actor visibility tracks.
+			// We do *that*, because it matches how Interchange always puts transform animations on the actors directly (so
+			// "scene component stuff" ends up as actor tracks), and also due to how it behaves better for cameras: Component
+			// visibility for camera nodes would hide the camera component itself, which has no effect. Actor visibility for
+			// camera actors does hide the entire camera actor however
+			if (!Info.bVisible)
+			{
+				SceneNode->SetCustomActorVisibility(Info.bVisible);
+			}
+
+			if (Info.bIsLODContainer)
+			{
+				TranslatorImpl.AddLODMeshNodes(Prim, NodeContainer, *SceneNode, TranslatorSettings, Info);
+			}
+
+			if (Info.bInsideLOD && Prim.IsA(TEXT("Mesh")))
+			{
+				TranslatorImpl.CurrentLODSceneNodes.Add(SceneNode);
+			}
+
+			// Skeleton joints are separate scene nodes in Interchange, so we need to emit that node hierarchy now
+			if (Prim.IsA(TEXT("Skeleton")))
+			{
+				AddSkeletonNodes(Prim, TranslatorImpl, *SceneNode, NodeContainer, Info);
+			}
+
+			if (Prim.IsA(TEXT("Volume")))
+			{
+				AddVolumeSceneNodeAttributes(
+					Prim,
+					SceneNode,
+					AssetNodeUids,
+					VolumeMaterialInstanceUid,
+					bNeedsVolumeTrack,
+					TranslatorImpl,
+					NodeContainer
+				);
+			}
+
+			// Connect scene node and primary asset node
+			if (AssetNodeUids.Num() != 0)
+			{
+				SceneNode->SetCustomAssetInstanceUid(AssetNodeUids[0]);
+			}
+
+			// Connect parent and child scene nodes
+			if (Info.ParentNode)
+			{
+				NodeContainer.SetNodeParentUid(SceneNode->GetUniqueID(), Info.ParentNode->GetUniqueID());
+			}
+
+			// Add animation tracks
+			if (bIsImageable)
+			{
+				AddPropertyAnimationNodes(Prim, TranslatorImpl, NodeContainer);
+				if (UsdUtils::HasAnimatedTransform(Prim))
+				{
+					AddTransformAnimationNode(Prim, TranslatorImpl, NodeContainer);
+				}
+			}
+		}
+
+		// Recurse into child prims
+		{
+			Info.ParentNode = SceneNode;
+
+			for (const FUsdPrim& ChildPrim : Prim.GetChildren())
+			{
+				Traverse(ChildPrim, TranslatorImpl, NodeContainer, TranslatorSettings, Info);
+			}
+		}
+
+		// Finalize the LOD container after recursing normally, because we'll rely on the regular traversal to
+		// process the mesh of the active LOD variant. AddLODMeshNodes only handles the inactive variants
+		if (Info.bIsLODContainer)
+		{
+			TranslatorImpl.FinalizeLODContainerTraversal(NodeContainer, Info, SceneNode);
+		}
+	}
+
+	bool GetStaticMeshPayloadData(
+		FString PayloadKey,
+		UInterchangeUSDTranslatorImpl& Impl,
+		const UsdToUnreal::FUsdMeshConversionOptions& Options,
+		FMeshDescription& OutMeshDescription
+	)
+	{
+		const bool bIsLODMesh = CheckAndChopPayloadPrefix(PayloadKey, LODPrefix);
+		const bool bIsPrimitiveShape = CheckAndChopPayloadPrefix(PayloadKey, PrimitiveShapePrefix);
+
+		const FString& PrimPath = PayloadKey;
+		UE::FUsdPrim Prim = Impl.UsdStage.GetPrimAtPath(UE::FSdfPath{*PrimPath});
+		if (!Prim)
+		{
+			if (bIsLODMesh)
+			{
+				Prim = Impl.TryGettingInactiveLODPrim(PrimPath);
+			}
+
+			if (!Prim)
+			{
+				return false;
+			}
+		}
+
+		FMeshDescription TempMeshDescription;
+		FStaticMeshAttributes StaticMeshAttributes(TempMeshDescription);
+		StaticMeshAttributes.Register();
+
+		// TODO: We can't do much with these yet: They will be used to generate primvar-compatible
+		// versions of the materials that are assigned to this mesh, whenever we get a pipeline
+		UsdUtils::FUsdPrimMaterialAssignmentInfo TempMaterialInfo;
+
+		bool bSuccess = false;
+		if (bIsPrimitiveShape)
+		{
+			bSuccess = UsdToUnreal::ConvertGeomPrimitive(Prim, TempMeshDescription, TempMaterialInfo, Options);
+		}
+		else
+		{
+			bSuccess = UsdToUnreal::ConvertGeomMesh(Prim, TempMeshDescription, TempMaterialInfo, Options);
+		}
+		if (!bSuccess)
+		{
+			return false;
+		}
+
+		UsdUtils::FUsdPrimMaterialAssignmentInfo* AssignmentPtr = &TempMaterialInfo;
+		if (bIsLODMesh)
+		{
+			// Use our cached material assignments instead of whatever we pull from ConvertGeomMesh
+			// because if we're in an LOD mesh then we may be reading from a temp stage, that has a
+			// population mask that may not include the material, meaning ConvertGeomMesh may have failed
+			// to resolve all the bindings. The cached assignments come from AddMeshNode step, where
+			// we switch the active variant on the current stage and so get nice material bindings that
+			// resolve normally.
+			//
+			// Note that we can't even use the info cache here, because it wouldn't have cached info
+			// about the inactive LOD variants
+			if (UsdUtils::FUsdPrimMaterialAssignmentInfo* CachedInfo = Impl.CachedMaterialAssignments.Find(PrimPath))
+			{
+				AssignmentPtr = CachedInfo;
+			}
+		}
+
+		OutMeshDescription = MoveTemp(TempMeshDescription);
+
+		FixMaterialSlotNames(OutMeshDescription, AssignmentPtr->Slots);
+
+		return true;
+	}
+
+	bool GetSkeletalMeshPayloadData(
+		FString PayloadKey,
+		UInterchangeUSDTranslatorImpl& Impl,
+		const UsdToUnreal::FUsdMeshConversionOptions& Options,
+		FMeshDescription& OutMeshDescription,
+		TArray<FString>& OutJointNames
+	)
+	{
+#if WITH_EDITOR
+		const bool bIsLODMesh = CheckAndChopPayloadPrefix(PayloadKey, LODPrefix);
+
+		const FString& PrimPath = PayloadKey;
+		UE::FUsdPrim Prim = Impl.UsdStage.GetPrimAtPath(UE::FSdfPath{*PrimPath});
+		if (!Prim)
+		{
+			if (bIsLODMesh)
+			{
+				Prim = Impl.TryGettingInactiveLODPrim(PrimPath);
+			}
+
+			if (!Prim)
+			{
+				return false;
+			}
+		}
+
+		const FString& MeshNodeUid = MeshPrefix + Prim.GetPrimPath().GetString();
+
+		// Read these variables from the data we cached during traversal for translation
+		TSharedPtr<TArray<FString>> JointNames = nullptr;
+		UE::FUsdSkelSkeletonQuery SkelQuery;
+		{
+			FReadScopeLock ReadLock{Impl.CachedTraversalInfoLock};
+
+			const FTraversalInfo* MeshInfo = Impl.NodeUidToCachedTraversalInfo.Find(MeshNodeUid);
+			if (!MeshInfo)
+			{
+				return false;
+			}
+
+			SkelQuery = MeshInfo->ResolveSkelQuery(Impl.UsdStage);
+			if (!SkelQuery)
+			{
+				return false;
+			}
+
+			// The above fields are associated to the mesh *asset* node Uid (hence the prefix),
+			// while the joint names are associated to the skeleton *scene* node Uid, so no prefix
+			const FString& SkeletonNodeUid = *MeshInfo->BoundSkeletonPrimPath;
+			const FTraversalInfo* SkeletonInfo = Impl.NodeUidToCachedTraversalInfo.Find(SkeletonNodeUid);
+			if (!SkeletonInfo)
+			{
+				return false;
+			}
+			JointNames = SkeletonInfo->SkelJointNames;
+			if (!JointNames)
+			{
+				return false;
+			}
+		}
+
+		UE::FUsdSkelSkinningQuery SkinningQuery = UsdUtils::CreateSkinningQuery(Prim, SkelQuery);
+		if (!SkinningQuery)
+		{
+			return false;
+		}
+
+		FSkeletalMeshImportData SkelMeshImportData;
+		UsdUtils::FUsdPrimMaterialAssignmentInfo TempMaterialInfo;
+		FMeshDescription TempMeshDescription;
+
+		pxr::UsdSkelSkinningQuery& PxrSkelSkinningQueryPtr = SkinningQuery;
+		pxr::UsdSkelSkeletonQuery& PxrSkelSkeletonQueryPtr = SkelQuery;
+
+		bool bSuccess = UsdToUnreal::ConvertGeomMesh(
+			Prim,
+			TempMeshDescription,
+			TempMaterialInfo,
+			Options,
+			&PxrSkelSkinningQueryPtr,
+			&PxrSkelSkeletonQueryPtr
+		);
+		if (!bSuccess)
+		{
+			return false;
+		}
+
+		OutMeshDescription = MoveTemp(TempMeshDescription);
+
+		UsdUtils::FUsdPrimMaterialAssignmentInfo* AssignmentPtr = &TempMaterialInfo;
+		if (bIsLODMesh)
+		{
+			// Use our cached material assignments instead of whatever we pull from ConvertSkinnedMesh
+			// because if we're in an LOD mesh then we may be reading from a temp stage, that has a
+			// population mask that may not include the material, meaning ConvertSkinnedMesh may have failed
+			// to resolve all the bindings. The cached assignments come from AddMeshNode step, where
+			// we switch the active variant on the current stage and so get nice material bindings that
+			// resolve normally.
+
+			// Note that we can't even use the info cache here, because it wouldn't have cached info
+			// about the inactive LOD variants
+			if (UsdUtils::FUsdPrimMaterialAssignmentInfo* CachedInfo = Impl.CachedMaterialAssignments.Find(PrimPath))
+			{
+				AssignmentPtr = CachedInfo;
+			}
+		}
+
+		FixMaterialSlotNames(OutMeshDescription, AssignmentPtr->Slots);
+
+		OutJointNames = *JointNames;
+
+		return true;
+#else
+		return false;
+#endif	  // WITH_EDITOR
+	}
+
+	bool GetMorphTargetPayloadData(
+		FString PayloadKey,
+		UInterchangeUSDTranslatorImpl& Impl,
+		const UsdToUnreal::FUsdMeshConversionOptions& Options,
+		FMeshDescription& OutMeshDescription,
+		FString& OutMorphTargetName
+	)
+	{
+		bool bIsLODMesh = false;
+		FString MeshPrimPath;
+		int32 BlendShapeIndex = INDEX_NONE;
+		FString InbetweenName;
+		if (!ParseMorphTargetMeshPayloadKey(PayloadKey, bIsLODMesh, MeshPrimPath, BlendShapeIndex, InbetweenName))
+		{
+			return false;
+		}
+
+		UE::FUsdPrim MeshPrim = Impl.UsdStage.GetPrimAtPath(FSdfPath{*MeshPrimPath});
+		if (!MeshPrim)
+		{
+			if (bIsLODMesh)
+			{
+				MeshPrim = Impl.TryGettingInactiveLODPrim(MeshPrimPath);
+			}
+
+			if (!MeshPrim)
+			{
+				return false;
+			}
+		}
+
+		UE::FUsdSkelBlendShapeQuery Query{MeshPrim};
+		if (!Query)
+		{
+			return false;
+		}
+
+		UE::FUsdSkelBlendShape BlendShape = Query.GetBlendShape(BlendShapeIndex);
+		if (!BlendShape)
+		{
+			return false;
+		}
+
+		// TODO: This is extremely slow, as it will reimport the mesh for every single morph target!
+		// It seems to be what the other translators do, however. We need some form of FMeshDescription caching here
+		TArray<FString> UnusedJointNames;
+		const FString MeshPayloadKey = bIsLODMesh ? LODPrefix + MeshPrimPath : MeshPrimPath;
+		bool bConverted = GetSkeletalMeshPayloadData(MeshPayloadKey, Impl, Options, OutMeshDescription, UnusedJointNames);
+		if (!bConverted || OutMeshDescription.IsEmpty())
+		{
+			return false;
+		}
+
+		OutMorphTargetName = BlendShape.GetPrim().GetName().ToString();
+		if (!InbetweenName.IsEmpty())
+		{
+			OutMorphTargetName += TEXT("_") + InbetweenName;
+		}
+
+		const float Weight = 1.0f;
+		return UsdUtils::ApplyBlendShape(OutMeshDescription, BlendShape.GetPrim(), Weight, InbetweenName);
+	}
+
+	// Volume animations need a special property reader because in USD it's an animation of the file path attribute within
+	// the asset prims, while in UE we want a float track flipping through the volume frame indices
+	UsdToUnreal::FPropertyTrackReader CreateVolumeTrackReader(
+		const FString& InPrimPath,
+		const UInterchangeUSDTranslatorImpl& InImpl,
+		TArray<double>& OutTimeSampleUnion
+	)
+	{
+		UsdToUnreal::FPropertyTrackReader Result;
+
+		const TArray<UsdUtils::FVolumePrimInfo>* CollectedInfoForPrim = InImpl.PrimPathToVolumeInfo.Find(InPrimPath);
+		if (!CollectedInfoForPrim)
+		{
+			return Result;
+		}
+
+		const UsdUtils::FVolumePrimInfo* AnimatedInfo = nullptr;
+
+		for (const UsdUtils::FVolumePrimInfo& Info : *CollectedInfoForPrim)
+		{
+			if (Info.TimeSamplePathTimeCodes.Num() > 0)
+			{
+				if (AnimatedInfo != nullptr)
+				{
+					UE_LOG(
+						LogUsd,
+						Log,
+						TEXT(
+							"Only one animated SparseVolumeTexture can be driven via LevelSequences for each prim, for now. Prim '%s' has multiple, so the animation may not be correct."
+						),
+						*InPrimPath
+					);
+				}
+
+				AnimatedInfo = &Info;
+			}
+		}
+
+		if (AnimatedInfo)
+		{
+			// Detect whether we inserted the default opinion volume as the first frame of the animation
+			bool bInsertedDefaultOpinion = AnimatedInfo->TimeSamplePaths.Num() > 0
+										   && AnimatedInfo->TimeSamplePaths[0] != AnimatedInfo->SourceVDBFilePath;
+
+			OutTimeSampleUnion = AnimatedInfo->TimeSamplePathTimeCodes;
+			Result.FloatReader = [AnimatedInfo, bInsertedDefaultOpinion](double TimeCode) -> float
+			{
+				int32 FrameIndex = 0;
+				for (; FrameIndex + 1 < AnimatedInfo->TimeSamplePathTimeCodes.Num(); ++FrameIndex)
+				{
+					if (AnimatedInfo->TimeSamplePathTimeCodes[FrameIndex + 1] > TimeCode)
+					{
+						break;
+					}
+				}
+				FrameIndex = FMath::Clamp(FrameIndex, 0, AnimatedInfo->TimeSamplePathTimeCodes.Num() - 1);
+
+				if (bInsertedDefaultOpinion)
+				{
+					FrameIndex += 1;
+				}
+
+				return static_cast<float>(FrameIndex);
+			};
+		}
+
+		return Result;
+	}
+
+	bool GetPropertyAnimationCurvePayloadData(
+		const UInterchangeUSDTranslatorImpl& Impl,
+		const FString& PayloadKey,
+		Interchange::FAnimationPayloadData& OutPayloadData
+	)
+	{
+		FString PrimPath;
+		FString UEPropertyNameStr;
+		bool bSplit = PayloadKey.Split(TEXT("\\"), &PrimPath, &UEPropertyNameStr, ESearchCase::CaseSensitive, ESearchDir::FromEnd);
+		if (!bSplit)
+		{
+			return false;
+		}
+
+		UE::FUsdPrim Prim = Impl.UsdStage.GetPrimAtPath(UE::FSdfPath{*PrimPath});
+		FName UEPropertyName = *UEPropertyNameStr;
+		if (!Prim || UEPropertyName == NAME_None)
+		{
+			return false;
+		}
+
+		TArray<double> TimeSampleUnion;
+		UsdToUnreal::FPropertyTrackReader Reader;
+
+		if (Prim.IsA(TEXT("Volume")) && UEPropertyNameStr == GET_MEMBER_NAME_CHECKED(UHeterogeneousVolumeComponent, Frame))
+		{
+			Reader = CreateVolumeTrackReader(PrimPath, Impl, TimeSampleUnion);
+		}
+		else
+		{
+			TArray<UE::FUsdAttribute> Attrs = UsdUtils::GetAttributesForProperty(Prim, UEPropertyName);
+			bool bSuccess = UE::FUsdAttribute::GetUnionedTimeSamples(Attrs, TimeSampleUnion);
+			if (!bSuccess)
+			{
+				return false;
+			}
+
+			const bool bIgnorePrimLocalTransform = false;
+			Reader = UsdToUnreal::CreatePropertyTrackReader(Prim, UEPropertyName, bIgnorePrimLocalTransform);
+		}
+
+		if (Reader.BoolReader)
+		{
+			return ReadBools(Impl.UsdStage, TimeSampleUnion, Reader.BoolReader, OutPayloadData);
+		}
+		else if (Reader.ColorReader)
+		{
+			return ReadColors(Impl.UsdStage, TimeSampleUnion, Reader.ColorReader, OutPayloadData);
+		}
+		else if (Reader.FloatReader)
+		{
+			return ReadFloats(Impl.UsdStage, TimeSampleUnion, Reader.FloatReader, OutPayloadData);
+		}
+		else if (Reader.TransformReader)
+		{
+			return ReadTransforms(Impl.UsdStage, TimeSampleUnion, Reader.TransformReader, OutPayloadData);
+		}
+
+		return false;
+	}
+
+	bool GetJointAnimationCurvePayloadData(
+		const UInterchangeUSDTranslatorImpl& Impl,
+		const TArray<const UE::Interchange::FAnimationPayloadQuery*>& Queries,
+		TArray<UE::Interchange::FAnimationPayloadData>& OutPayloadData
+	)
+	{
+		if (Queries.Num() == 0)
+		{
+			return false;
+		}
+
+		// We expect all queries to be for the same skeleton, and have the same timing parameters,
+		// since they were grouped up by HashAnimPayloadQuery, so let's just grab one for the params
+		const UE::Interchange::FAnimationPayloadQuery* FirstQuery = Queries[0];
+
+		// Parse payload key.
+		// Here is takes the form "<skeleton prim path>\<joint index in skeleton order>"
+		TArray<FString> PayloadKeyTokens;
+		FirstQuery->PayloadKey.UniqueId.ParseIntoArray(PayloadKeyTokens, TEXT("\\"));
+		if (PayloadKeyTokens.Num() != 2)
+		{
+			return false;
+		}
+
+		// Fetch our cached skeleton query
+		const FString& SkeletonPrimPath = PayloadKeyTokens[0];
+		UE::FUsdSkelSkeletonQuery SkelQuery;
+		{
+			FReadScopeLock ReadLock{Impl.CachedTraversalInfoLock};
+
+			const FTraversalInfo* MeshInfo = Impl.NodeUidToCachedTraversalInfo.Find(SkeletonPrimPath);
+			if (!MeshInfo)
+			{
+				return false;
+			}
+
+			SkelQuery = MeshInfo->ResolveSkelQuery(Impl.UsdStage);
+			if (!SkelQuery)
+			{
+				return false;
+			}
+		}
+
+		UE::FUsdPrim SkeletonPrim = SkelQuery.GetPrim();
+		UE::FUsdStage Stage = SkeletonPrim.GetStage();
+		FUsdStageInfo StageInfo{Stage};
+
+		// Compute the bake ranges and intervals
+		double TimeCodesPerSecond = Stage.GetTimeCodesPerSecond();
+		double BakeFrequency = FirstQuery->TimeDescription.BakeFrequency;
+		double RangeStartSeconds = FirstQuery->TimeDescription.RangeStartSecond;
+		double RangeStopSeconds = FirstQuery->TimeDescription.RangeStopSecond;
+		double SectionLengthSeconds = RangeStopSeconds - RangeStartSeconds;
+		double StartTimeCode = RangeStartSeconds * TimeCodesPerSecond;
+		const int32 NumBakedFrames = FMath::RoundToInt(FMath::Max(SectionLengthSeconds * TimeCodesPerSecond + 1.0, 1.0));
+		double TimeCodeIncrement = (1.0 / BakeFrequency) * TimeCodesPerSecond;
+
+		// Bake all joint transforms via USD into arrays for each separate joint (in whatever order SkelQuery gives us)
+		TArray<TArray<FTransform>> BakedTransforms;
+		for (int32 FrameIndex = 0; FrameIndex < NumBakedFrames; ++FrameIndex)
+		{
+			const double FrameTimeCode = StartTimeCode + FrameIndex * TimeCodeIncrement;
+
+			TArray<FTransform> TransformsForTimeCode;
+			bool bSuccess = SkelQuery.ComputeJointLocalTransforms(TransformsForTimeCode, FrameTimeCode);
+			if (!bSuccess)
+			{
+				break;
+			}
+
+			for (FTransform& Transform : TransformsForTimeCode)
+			{
+				Transform = UsdUtils::ConvertTransformToUESpace(StageInfo, Transform);
+			}
+
+			// Setup our BakedTransforms in here, because we may actually get more or less transforms
+			// from the SkeletonQuery than our AnimSequence wants/expects, given that it can specify
+			// its own animated joint order
+			int32 NumSkelJoints = TransformsForTimeCode.Num();
+			if (FrameIndex == 0)
+			{
+				BakedTransforms.SetNum(NumSkelJoints);
+				for (int32 JointIndex = 0; JointIndex < NumSkelJoints; ++JointIndex)
+				{
+					BakedTransforms[JointIndex].SetNum(NumBakedFrames);
+				}
+			}
+
+			// Transpose our baked transforms into the arrays we'll eventually return
+			for (int32 JointIndex = 0; JointIndex < NumSkelJoints; ++JointIndex)
+			{
+				BakedTransforms[JointIndex][FrameIndex] = TransformsForTimeCode[JointIndex];
+			}
+		}
+
+		// Finally build our payload data return values by picking the desired baked arrays with the payload joint indices
+		OutPayloadData.Reset(Queries.Num());
+		for (int32 QueryIndex = 0; QueryIndex < Queries.Num(); ++QueryIndex)
+		{
+			const UE::Interchange::FAnimationPayloadQuery* Query = Queries[QueryIndex];
+
+			FString IndexStr = Query->PayloadKey.UniqueId.RightChop(SkeletonPrimPath.Len() + 1);	// Also skip the '\'
+			int32 JointIndex = INDEX_NONE;
+			bool bLexed = LexTryParseString(JointIndex, *IndexStr);
+			if (!bLexed)
+			{
+				continue;
+			}
+
+			UE::Interchange::FAnimationPayloadData& PayloadData = OutPayloadData.Emplace_GetRef(Query->SceneNodeUniqueID, Query->PayloadKey);
+			PayloadData.BakeFrequency = BakeFrequency;
+			PayloadData.RangeStartTime = RangeStartSeconds;
+			PayloadData.RangeEndTime = RangeStopSeconds;
+
+			if (BakedTransforms.IsValidIndex(JointIndex))
+			{
+				PayloadData.Transforms = MoveTemp(BakedTransforms[JointIndex]);
+			}
+		}
+
+		return true;
+	}
+
+	bool GetMorphTargetAnimationCurvePayloadData(
+		const UInterchangeUSDTranslatorImpl& Impl,
+		const FString& PayloadKey,
+		Interchange::FAnimationPayloadData& OutPayloadData
+	)
+	{
+		// Here we must output the morph target curve for a particular channel and skinning target, i.e.
+		// the connection of a SkelAnimation blend shape channel to a particular Mesh prim.
+
+		// These payload keys were generated from GetMorphTargetCurvePayloadKey(), so they take the form
+		// "<skeleton prim path>\<skel anim channel index>\<blend shape path>"
+		TArray<FString> PayloadKeyTokens;
+		PayloadKey.ParseIntoArray(PayloadKeyTokens, TEXT("\\"));
+		if (PayloadKeyTokens.Num() != 3)
+		{
+			return false;
+		}
+		const FString& SkeletonPrimPath = PayloadKeyTokens[0];
+		const FString& AnimChannelIndexStr = PayloadKeyTokens[1];
+		const FString& BlendShapePath = PayloadKeyTokens[2];
+
+		const UE::FUsdStage& UsdStage = Impl.UsdStage;
+
+		int32 SkelAnimChannelIndex = INDEX_NONE;
+		bool bLexed = LexTryParseString(SkelAnimChannelIndex, *AnimChannelIndexStr);
+
+		UE::FUsdPrim BlendShapePrim = UsdStage.GetPrimAtPath(UE::FSdfPath{*BlendShapePath});
+		UE::FUsdSkelBlendShape BlendShape{BlendShapePrim};
+		if (!BlendShape || !bLexed || SkelAnimChannelIndex == INDEX_NONE)
+		{
+			return false;
+		}
+		const FString BlendShapeName = BlendShapePrim.GetName().ToString();
+
+		// Fill in the actual morph target curve
+		UE::FUsdSkelAnimQuery AnimQuery;
+		{
+			UE::FUsdSkelSkeletonQuery SkelQuery;
+			{
+				FReadScopeLock ReadLock{Impl.CachedTraversalInfoLock};
+
+				const FTraversalInfo* MeshInfo = Impl.NodeUidToCachedTraversalInfo.Find(SkeletonPrimPath);
+				if (!MeshInfo)
+				{
+					return false;
+				}
+
+				SkelQuery = MeshInfo->ResolveSkelQuery(Impl.UsdStage);
+				if (!SkelQuery)
+				{
+					return false;
+				}
+			}
+
+			AnimQuery = SkelQuery.GetAnimQuery();
+			if (!AnimQuery)
+			{
+				return false;
+			}
+
+			TArray<double> TimeCodes;
+			bool bSuccess = AnimQuery.GetBlendShapeWeightTimeSamples(TimeCodes);
+			if (!bSuccess)
+			{
+				return false;
+			}
+
+			OutPayloadData.Curves.SetNum(1);
+			FRichCurve& Curve = OutPayloadData.Curves[0];
+			Curve.ReserveKeys(TimeCodes.Num());
+
+			const FFrameRate StageFrameRate{static_cast<uint32>(UsdStage.GetTimeCodesPerSecond()), 1};
+			const ERichCurveInterpMode InterpMode = (UsdStage.GetInterpolationType() == EUsdInterpolationType::Linear)
+														? ERichCurveInterpMode::RCIM_Linear
+														: ERichCurveInterpMode::RCIM_Constant;
+
+			TArray<float> Weights;
+			for (double TimeCode : TimeCodes)
+			{
+				bSuccess = AnimQuery.ComputeBlendShapeWeights(Weights, TimeCode);
+				if (!bSuccess || !Weights.IsValidIndex(SkelAnimChannelIndex))
+				{
+					break;
+				}
+
+				int32 FrameNumber = FMath::FloorToInt(TimeCode);
+				float SubFrameNumber = TimeCode - FrameNumber;
+
+				FFrameTime FrameTime{FrameNumber, SubFrameNumber};
+				double FrameTimeSeconds = static_cast<float>(StageFrameRate.AsSeconds(FrameTime));
+
+				FKeyHandle Handle = Curve.AddKey(FrameTimeSeconds, Weights[SkelAnimChannelIndex]);
+				Curve.SetKeyInterpMode(Handle, InterpMode);
+			}
+		}
+
+		TArray<FString> SkelAnimChannels = AnimQuery.GetBlendShapeOrder();
+
+		// Provide inbetween names/positions for this morph target payload
+		TArray<UE::FUsdSkelInbetweenShape> Inbetweens = BlendShape.GetInbetweens();
+		if (Inbetweens.Num() > 0)
+		{
+			// Let's store them into this temp struct so that we can sort them by weight first,
+			// as Interchange seems to expect that given how it will pass these right along into
+			// ResolveWeightsForBlendShape inside InterchangeAnimSequenceFactory.cpp
+			struct FInbetweenAndPosition
+			{
+				FString Name;
+				float Position;
+			};
+			TArray<FInbetweenAndPosition> ParsedInbetweens;
+			ParsedInbetweens.Reset(Inbetweens.Num());
+
+			for (const UE::FUsdSkelInbetweenShape& Inbetween : Inbetweens)
+			{
+				float Position = 0.5f;
+				bool bSuccess = Inbetween.GetWeight(&Position);
+				if (!bSuccess)
+				{
+					continue;
+				}
+
+				// Skip invalid positions. Note that technically positions outside the [0, 1] range seem to be allowed, but
+				// they don't seem to work very well with our inbetween weights resolution function for some reason.
+				// The legacy USD workflows have this exact same check though, so for consistency let's just do the same, and
+				// if becomes an issue we should fix both
+				if (Position > 1.0f || Position < 0.0f || FMath::IsNearlyZero(Position) || FMath::IsNearlyEqual(Position, 1.0f))
+				{
+					continue;
+				}
+
+				const FString MorphTargetName = BlendShapeName + TEXT("_") + Inbetween.GetAttr().GetName().ToString();
+				FInbetweenAndPosition& NewEntry = ParsedInbetweens.Emplace_GetRef();
+				NewEntry.Name = MorphTargetName;
+				NewEntry.Position = Position;
+			}
+
+			ParsedInbetweens.Sort(
+				[](const FInbetweenAndPosition& LHS, const FInbetweenAndPosition& RHS)
+				{
+					// It's invalid USD to author two inbetweens with the same weight, so let's ignore that case here.
+					// (Reference: https://openusd.org/release/api/_usd_skel__schemas.html#UsdSkel_BlendShape)
+					return LHS.Position < RHS.Position;
+				}
+			);
+
+			OutPayloadData.InbetweenCurveNames.Reset(Inbetweens.Num() + 1);
+			OutPayloadData.InbetweenFullWeights.Reset(Inbetweens.Num());
+
+			// We add the main morph target curve name to InbetweenCurveNames too (having it end up one size bigger than
+			// InbetweenFullWeights) as it seems like that's what Interchange expects. See CreateMorphTargetCurve within
+			// InterchangeAnimSequenceFactory.cpp, and the very end of FFbxMesh::AddAllMeshes within FbxMesh.cpp
+			OutPayloadData.InbetweenCurveNames.Add(BlendShapeName);
+
+			for (const FInbetweenAndPosition& InbetweenAndPosition : ParsedInbetweens)
+			{
+				OutPayloadData.InbetweenCurveNames.Add(InbetweenAndPosition.Name);
+				OutPayloadData.InbetweenFullWeights.Add(InbetweenAndPosition.Position);
+			}
+		}
+
+		return true;
+	}
+#else
+	class UInterchangeUSDTranslatorImpl
+	{
+	};
+#endif	  // USE_USD_SDK
+}	 // namespace UE::InterchangeUsdTranslator::Private
+
+UInterchangeUsdTranslatorSettings::UInterchangeUsdTranslatorSettings()
+	: GeometryPurpose((int32)(EUsdPurpose::Default | EUsdPurpose::Proxy | EUsdPurpose::Render | EUsdPurpose::Guide))
+	, RenderContext(UnrealIdentifiers::UnrealRenderContext)
+	, MaterialPurpose(*UnrealIdentifiers::MaterialPreviewPurpose)
+	, InterpolationType(EUsdInterpolationType::Linear)
+	, bOverrideStageOptions(false)
+	, StageOptions{
+		  0.01,				   // MetersPerUnit
+		  EUsdUpAxis::ZAxis	   // UpAxis
+	  }
+	, bTranslatePrimAttributes(false) // False by default as it could be expensive to traverse all attributes of all prims running the regex
+	, AttributeRegexFilter(TEXT("."))
+{
+}
+
+UInterchangeUSDTranslator::UInterchangeUSDTranslator()
+	: Impl(MakeUnique<UE::InterchangeUsdTranslator::Private::UInterchangeUSDTranslatorImpl>())
+{
+}
+
+EInterchangeTranslatorType UInterchangeUSDTranslator::GetTranslatorType() const
+{
+	return GInterchangeEnableUSDLevelImport ? EInterchangeTranslatorType::Scenes : EInterchangeTranslatorType::Assets;
+}
+
+EInterchangeTranslatorAssetType UInterchangeUSDTranslator::GetSupportedAssetTypes() const
+{
+	return EInterchangeTranslatorAssetType::Materials | EInterchangeTranslatorAssetType::Meshes | EInterchangeTranslatorAssetType::Animations;
+}
+
+TArray<FString> UInterchangeUSDTranslator::GetSupportedFormats() const
+{
+	TArray<FString> Extensions;
+	if (GInterchangeEnableUSDImport)
+	{
+		FModuleManager::Get().LoadModuleChecked(TEXT("UnrealUSDWrapper"));
+		UnrealUSDWrapper::AddUsdImportFileFormatDescriptions(Extensions);
+	}
+	return Extensions;
+}
+
+bool UInterchangeUSDTranslator::Translate(UInterchangeBaseNodeContainer& NodeContainer) const
+{
+#if USE_USD_SDK
+	using namespace UE;
+	using namespace UE::InterchangeUsdTranslator::Private;
+	using namespace UsdToUnreal;
+
+	// Reset impl as we don't want to share internal state with a previous translation, as the file
+	// (or our settings) may have changed, which should lead to different nodes / data
+	Impl = MakeUnique<UE::InterchangeUsdTranslator::Private::UInterchangeUSDTranslatorImpl>();
+	UInterchangeUSDTranslatorImpl* ImplPtr = Impl.Get();
+	if (!ImplPtr)
+	{
+		return false;
+	}
+
+	UInterchangeUsdTranslatorSettings* Settings = Cast<UInterchangeUsdTranslatorSettings>(GetSettings());
+	if (!Settings)
+	{
+		return false;
+	}
+
+	if (!SourceData)
+	{
+		return false;
+	}
+
+	// Setup context
+	UObject* ContextObject = SourceData->GetContextObjectByTag(UE::Interchange::USD::USDContextTag);
+	UInterchangeUsdContext* Context = Cast<UInterchangeUsdContext>(ContextObject);
+	if (!Context)
+	{
+		Context = NewObject<UInterchangeUsdContext>();
+
+		ensureMsgf(
+			!ContextObject,
+			TEXT("Invalid ContextObject with tag '%s' will be removed and replaced with an UInterchangeUsdContext object"),
+			*UE::Interchange::USD::USDContextTag
+		);
+		SourceData->SetContextObjectByTag(UE::Interchange::USD::USDContextTag, Context);
+	}
+
+	// Setup stage
+	FString USDZFilePath;
+	FString DecompressedUSDZRoot;
+	UE::FUsdStage StageToImport = Context->GetUsdStage();
+	{
+		// Context didn't provide a stage: Try loading one from the provided file path
+		if (!StageToImport)
+		{
+			FString FilePath = SourceData->GetFilename();
+			if (!FPaths::FileExists(FilePath))
+			{
+				return false;
+			}
+
+			// If we're provided a USDZ path, for now we will decompress it into a temp dir and redirect our paths.
+			//
+			// This is mainly because the texture factories must receive a simple file path in order to produce the
+			// their payloads. It's not practical to make them handle USDZ files, and it's not yet possible
+			// to provide them with raw binary buffers directly either
+			if (DecompressUSDZFileToTempFolder(FilePath, DecompressedUSDZRoot))
+			{
+				USDZFilePath = FilePath;
+				FilePath = DecompressedUSDZRoot;
+			}
+
+			// Import should always feel like it's directly from disk, so we ignore already loaded layers and stage cache
+			const bool bUseStageCache = false;
+			const bool bForceReloadLayersFromDisk = true;
+			StageToImport = UnrealUSDWrapper::OpenStage(*FilePath, EUsdInitialLoadSet::LoadAll, bUseStageCache, bForceReloadLayersFromDisk);
+		}
+
+		if (!StageToImport)
+		{
+			return false;
+		}
+
+		if (Settings)
+		{
+			// Apply coordinate system conversion to the stage if we have one
+			if (Settings->bOverrideStageOptions)
+			{
+				UsdUtils::SetUsdStageMetersPerUnit(StageToImport, Settings->StageOptions.MetersPerUnit);
+				UsdUtils::SetUsdStageUpAxis(StageToImport, Settings->StageOptions.UpAxis);
+			}
+
+			StageToImport.SetInterpolationType(Settings->InterpolationType);
+		}
+	}
+
+	// Setup info cache
+	FUsdInfoCache* InfoCache = Context->GetInfoCache();
+	if (!InfoCache)
+	{
+		InfoCache = Context->CreateOwnedInfoCache();
+	}
+
+	// Fill in our context with our potentially-created-on-demand stage and info cache
+	ensure(Context->SetUsdStage(StageToImport));
+	Context->SetExternalInfoCache(*InfoCache);
+
+	// Setup impl
+	ImplPtr->UsdStage = StageToImport;
+	ImplPtr->InfoCache = InfoCache;
+	ImplPtr->SetupTranslationContext(*Settings);
+	ImplPtr->TranslationContext->UsdInfoCache = ImplPtr->InfoCache;
+	ImplPtr->CurrentTrackSet = nullptr;
+	ImplPtr->USDZFilePath = USDZFilePath;
+	ImplPtr->DecompressedUSDZRoot = DecompressedUSDZRoot;
+
+	// Cache these so we don't have to keep converting these tokens over and over during translation
+	{
+		FUsdMeshConversionOptions& MeshOptions = ImplPtr->CachedMeshConversionOptions;
+
+		// We filter for this on the pipeline now
+		MeshOptions.PurposesToLoad = EUsdPurpose::Default | EUsdPurpose::Proxy | EUsdPurpose::Render | EUsdPurpose::Guide;
+
+		// TODO: Change FUsdMeshConversionOptions to not hold USD types directly, so we don't have to the conversion below everywhere.
+		// We can't use UsdToUnreal::ConvertToken() here because it returns a TUsdStore, and the template instantiation created in this module doesn't
+		// really do anything anyway as the module doesn't use IMPLEMENT_MODULE_USD!
+		// Luckily we can get around this here because pxr::TfToken doesn't allocate on its own: At most USD makes a copy of the string, which it
+		// should allocate/deallocate on its own allocator.
+		MeshOptions.RenderContext = Settings->RenderContext == UnrealIdentifiers::UniversalRenderContext
+										? pxr::UsdShadeTokens->universalRenderContext
+										: pxr::TfToken{TCHAR_TO_UTF8(*Settings->RenderContext.ToString())};
+		MeshOptions.MaterialPurpose = Settings->MaterialPurpose.IsNone() ? pxr::UsdShadeTokens->allPurpose
+																		 : pxr::TfToken{TCHAR_TO_UTF8(*Settings->MaterialPurpose.ToString())};
+	}
+
+	// Traverse stage and emit translated nodes
+	{
+		ImplPtr->InfoCache->RebuildCacheForSubtrees({UE::FSdfPath::AbsoluteRootPath()}, *ImplPtr->TranslationContext);
+
+		FTraversalInfo Info;
+		Traverse(ImplPtr->UsdStage.GetPseudoRoot(), *ImplPtr, NodeContainer, TranslatorSettings, Info);
+	}
+
+	return true;
+#else
+	return false;
+#endif	  // USE_USD_SDK
+}
+
+void UInterchangeUSDTranslator::ReleaseSource()
+{
+#if USE_USD_SDK
+	UE::InterchangeUsdTranslator::Private::UInterchangeUSDTranslatorImpl* ImplPtr = Impl.Get();
+	if (ImplPtr)
+	{
+		ImplPtr->TranslationContext.Reset();	// This holds a reference to the stage
+		ImplPtr->UsdStage = UE::FUsdStage{};
+		ImplPtr->InfoCache = nullptr;
+		ImplPtr->CurrentTrackSet = nullptr;
+
+		{
+			FWriteScopeLock Lock{ImplPtr->PrimPathToVariantToStageLock};
+			ImplPtr->PrimPathToVariantToStage.Reset();
+		}
+	}
+#endif	  // USE_USD_SDK
+
+	if (TranslatorSettings)
+	{
+		TranslatorSettings->ClearFlags(RF_Standalone);
+		TranslatorSettings = nullptr;
+	}
+}
+
+void UInterchangeUSDTranslator::ImportFinish()
+{
+#if USE_USD_SDK
+	UE::InterchangeUsdTranslator::Private::UInterchangeUSDTranslatorImpl* ImplPtr = Impl.Get();
+	if (!ImplPtr)
+	{
+		return;
+	}
+
+	ImplPtr->CleanUpDecompressedUSDZFolder();
+#endif	  // USE_USD_SDK
+}
+
+UInterchangeTranslatorSettings* UInterchangeUSDTranslator::GetSettings() const
+{
+	using namespace UE::InterchangeUsdTranslator::Private;
+
+	if (!TranslatorSettings)
+	{
+		TranslatorSettings = DuplicateObject<UInterchangeUsdTranslatorSettings>(
+			UInterchangeUsdTranslatorSettings::StaticClass()->GetDefaultObject<UInterchangeUsdTranslatorSettings>(),
+			GetTransientPackage()
+		);
+		TranslatorSettings->LoadSettings();
+		TranslatorSettings->ClearFlags(RF_ArchetypeObject);
+		TranslatorSettings->SetFlags(RF_Standalone);
+		TranslatorSettings->ClearInternalFlags(EInternalObjectFlags::Async);
+	}
+	return TranslatorSettings;
+}
+
+void UInterchangeUSDTranslator::SetSettings(const UInterchangeTranslatorSettings* InterchangeTranslatorSettings)
+{
+	using namespace UE::InterchangeUsdTranslator::Private;
+
+	if (TranslatorSettings)
+	{
+		TranslatorSettings->ClearFlags(RF_Standalone);
+		TranslatorSettings->ClearInternalFlags(EInternalObjectFlags::Async);
+		TranslatorSettings = nullptr;
+	}
+	if (const UInterchangeUsdTranslatorSettings* USDTranslatorSettings = Cast<UInterchangeUsdTranslatorSettings>(InterchangeTranslatorSettings))
+	{
+		TranslatorSettings = DuplicateObject<UInterchangeUsdTranslatorSettings>(USDTranslatorSettings, GetTransientPackage());
+		TranslatorSettings->ClearInternalFlags(EInternalObjectFlags::Async);
+		TranslatorSettings->SetFlags(RF_Standalone);
+	}
+}
+
+TOptional<UE::Interchange::FMeshPayloadData> UInterchangeUSDTranslator::GetMeshPayloadData(
+	const FInterchangeMeshPayLoadKey& PayloadKey,
+	const FTransform& MeshGlobalTransform
+) const
+{
+	using namespace UE::InterchangeUsdTranslator::Private;
+	bool bSuccess = false;
+	TOptional<UE::Interchange::FMeshPayloadData> Result;
+#if USE_USD_SDK
+	UInterchangeUSDTranslatorImpl* ImplPtr = Impl.Get();
+	if (!ImplPtr)
+	{
+		return Result;
+	}
+
+	UsdToUnreal::FUsdMeshConversionOptions OptionsCopy = ImplPtr->CachedMeshConversionOptions;
+	OptionsCopy.AdditionalTransform = MeshGlobalTransform;
+	OptionsCopy.bMergeIdenticalMaterialSlots = false;	 // This must always be false, because we need the material assignments we read from the
+														 // meshes to match up with whatever we cached from AddMeshNode, in order to fixup LOD
+														 // material slots
+
+	UE::Interchange::FMeshPayloadData MeshPayloadData;
+	switch (PayloadKey.Type)
+	{
+		case EInterchangeMeshPayLoadType::STATIC:
+		{
+			bSuccess = UE::InterchangeUsdTranslator::Private::GetStaticMeshPayloadData(
+				PayloadKey.UniqueId,
+				*ImplPtr,
+				OptionsCopy,
+				MeshPayloadData.MeshDescription
+			);
+			break;
+		}
+		case EInterchangeMeshPayLoadType::SKELETAL:
+		{
+			bSuccess = UE::InterchangeUsdTranslator::Private::GetSkeletalMeshPayloadData(
+				PayloadKey.UniqueId,
+				*ImplPtr,
+				OptionsCopy,
+				MeshPayloadData.MeshDescription,
+				MeshPayloadData.JointNames
+			);
+			break;
+		}
+		case EInterchangeMeshPayLoadType::MORPHTARGET:
+		{
+			bSuccess = UE::InterchangeUsdTranslator::Private::GetMorphTargetPayloadData(
+				PayloadKey.UniqueId,
+				*ImplPtr,
+				OptionsCopy,
+				MeshPayloadData.MeshDescription,
+				MeshPayloadData.MorphTargetName
+			);
+			break;
+		}
+		case EInterchangeMeshPayLoadType::ANIMATED:
+		{
+			OptionsCopy.TimeCode = PayloadKey.FrameNumber;
+
+			bSuccess = UE::InterchangeUsdTranslator::Private::GetStaticMeshPayloadData(
+				PayloadKey.UniqueId,
+				*ImplPtr,
+				OptionsCopy,
+				MeshPayloadData.MeshDescription
+			);
+			break;
+		}
+		case EInterchangeMeshPayLoadType::NONE:	   // Fallthrough
+		default:
+			break;
+	}
+
+	if (bSuccess)
+	{
+		Result.Emplace(MeshPayloadData);
+	}
+#endif	  // USE_USD_SDK
+	return Result;
+}
+
+namespace UE::InterchangeUsdTranslator::Private
+{
+	TOptional<UE::Interchange::FImportImage> GetMaterialXTexturePayloadData(
+		const FString& PayloadKey,
+		TOptional<FString>& AlternateTexturePath,
+		const UInterchangeTranslatorBase* Translator,
+		UInterchangeUSDTranslatorImpl& Impl
+	)
+	{
+		TOptional<UE::Interchange::FImportImage> TexturePayloadData;
+
+#if USE_USD_SDK
+		FString Filename = PayloadKey;
+		TextureCompressionSettings CompressionSettings = TextureCompressionSettings::TC_Default;
+
+#if WITH_EDITOR
+		if (int32 IndexTextureCompression; PayloadKey.FindChar(FMaterialXManager::TexturePayloadSeparator, IndexTextureCompression))
+		{
+			Filename = PayloadKey.Mid(0, IndexTextureCompression);
+			CompressionSettings = TextureCompressionSettings(FCString::Atoi(*PayloadKey.Mid(IndexTextureCompression + 1)));
+		}
+#endif
+		UInterchangeSourceData* PayloadSourceData = UInterchangeManager::GetInterchangeManager().CreateSourceData(Filename);
+		FGCObjectScopeGuard ScopedSourceData(PayloadSourceData);
+
+		if (!PayloadSourceData)
+		{
+			return TexturePayloadData;
+		}
+
+		UInterchangeTranslatorBase* SourceTranslator = UInterchangeManager::GetInterchangeManager().GetTranslatorForSourceData(PayloadSourceData);
+		FGCObjectScopeGuard ScopedSourceTranslator(SourceTranslator);
+		const IInterchangeTexturePayloadInterface* TextureTranslator = Cast<IInterchangeTexturePayloadInterface>(SourceTranslator);
+
+		if (!TextureTranslator)
+		{
+			return TexturePayloadData;
+		}
+
+		SourceTranslator->SetResultsContainer(Translator->Results);
+
+		AlternateTexturePath = Impl.GetTextureSourcePath(Filename);
+
+		TexturePayloadData = TextureTranslator->GetTexturePayloadData(PayloadKey, AlternateTexturePath);
+		if (TexturePayloadData.IsSet())
+		{
+			TexturePayloadData.GetValue().CompressionSettings = CompressionSettings;
+		}
+#endif	  // USE_USD_SDK
+
+		return TexturePayloadData;
+	}
+}	 // namespace UE::InterchangeUsdTranslator::Private
+
+TOptional<UE::Interchange::FImportImage> UInterchangeUSDTranslator::GetTexturePayloadData(
+	const FString& PayloadKey,
+	TOptional<FString>& AlternateTexturePath
+) const
+{
+	using namespace UE::InterchangeUsdTranslator::Private;
+
+	TOptional<UE::Interchange::FImportImage> TexturePayloadData;
+#if USE_USD_SDK
+
+	UInterchangeUSDTranslatorImpl* ImplPtr = Impl.Get();
+	if (!ImplPtr)
+	{
+		return TexturePayloadData;
+	}
+
+	FString FilePath;
+	TextureGroup TextureGroup;
+	bool bDecoded = DecodeTexturePayloadKey(PayloadKey, FilePath, TextureGroup);
+	if (bDecoded)
+	{
+		// Defer back to another translator to actually parse the texture raw data
+		UE::Interchange::Private::FScopedTranslator ScopedTranslator(FilePath, Results);
+		const IInterchangeTexturePayloadInterface* TextureTranslator = ScopedTranslator.GetPayLoadInterface<IInterchangeTexturePayloadInterface>();
+		if (TextureTranslator)
+		{
+			AlternateTexturePath = ImplPtr->GetTextureSourcePath(FilePath);
+
+			// The texture translators don't use the payload key, and read the texture directly from the SourceData's file path
+			const FString UnusedPayloadKey = {};
+			TexturePayloadData = TextureTranslator->GetTexturePayloadData(UnusedPayloadKey, AlternateTexturePath);
+
+			// Move compression settings onto the payload data.
+			// Note: We don't author anything else on the texture payload data here (like the sRGB flag), because those
+			// settings were already on our translated node, and presumably already made their way to the factory node.
+			// The factory should use them to override whatever it finds in this payload data, with the exception of the
+			// compression settings (which can't be stored on the translated node)
+			TexturePayloadData->CompressionSettings = TextureGroup == TEXTUREGROUP_WorldNormalMap ? TC_Normalmap : TC_Default;
+		}
+	}
+
+	// We did not find a suitable Payload in USD Translator, let's find one in one of the Translators (MaterialX for the moment)
+	// The best way would be to have a direct association between the payload and the right Translator, but we don't have a suitable way of knowing
+	// which Payload belongs to which Translator So let's just loop over them all
+	if (!TexturePayloadData.IsSet())
+	{
+		for (const TPair<FString, TStrongObjectPtr<UInterchangeTranslatorBase>>& Pair : Impl->Translators)
+		{
+			if (IInterchangeTexturePayloadInterface* TexturePayloadInterface = Cast<IInterchangeTexturePayloadInterface>(Pair.Value.Get()))
+			{
+				TexturePayloadData = TexturePayloadInterface->GetTexturePayloadData(PayloadKey, AlternateTexturePath);
+				if (TexturePayloadData.IsSet())
+				{
+					break;
+				}
+			}
+		}
+	}
+
+	// If we couldn't find a texture in either the USD Translator nor the Translators, then it's most likely coming from reading an mtlx in memory,
+	// copying the behavior from InterchangeMaterialXTranslator
+	if (!TexturePayloadData.IsSet())
+	{
+		TexturePayloadData = GetMaterialXTexturePayloadData(PayloadKey, AlternateTexturePath, this, *ImplPtr);
+	}
+
+#endif	  // USE_USD_SDK
+	return TexturePayloadData;
+}
+
+TOptional<UE::Interchange::FImportBlockedImage> UInterchangeUSDTranslator::GetBlockedTexturePayloadData(
+	const FString& PayloadKey,
+	TOptional<FString>& AlternateTexturePath
+) const
+{
+	using namespace UE::InterchangeUsdTranslator::Private;
+
+	UE::Interchange::FImportBlockedImage BlockData;
+
+#if USE_USD_SDK
+	UInterchangeUSDTranslatorImpl* ImplPtr = Impl.Get();
+	if (!ImplPtr)
+	{
+		return {};
+	}
+
+	FString FilePath;
+	TextureGroup TextureGroup;
+	bool bDecoded = DecodeTexturePayloadKey(PayloadKey, FilePath, TextureGroup);
+	if (!bDecoded)
+	{
+		return {};
+	}
+
+	AlternateTexturePath = ImplPtr->GetTextureSourcePath(FilePath);
+
+	// Collect all the UDIM tile filepaths similar to this current tile. If we've been asked to translate
+	// a blocked texture then we must have some
+	TMap<int32, FString> TileIndexToPath = UE::TextureUtilitiesCommon::GetUDIMBlocksFromSourceFile(
+		FilePath,
+		UE::TextureUtilitiesCommon::DefaultUdimRegexPattern
+	);
+	if (!ensure(TileIndexToPath.Num() > 0))
+	{
+		return {};
+	}
+
+	bool bInitializedBlockData = false;
+
+	TArray<UE::Interchange::FImportImage> TileImages;
+	TileImages.Reserve(TileIndexToPath.Num());
+
+	for (const TPair<int32, FString>& TileIndexAndPath : TileIndexToPath)
+	{
+		int32 UdimTile = TileIndexAndPath.Key;
+		const FString& TileFilePath = TileIndexAndPath.Value;
+
+		int32 BlockX = INDEX_NONE;
+		int32 BlockY = INDEX_NONE;
+		UE::TextureUtilitiesCommon::ExtractUDIMCoordinates(UdimTile, BlockX, BlockY);
+		if (BlockX == INDEX_NONE || BlockY == INDEX_NONE)
+		{
+			continue;
+		}
+
+		// Find another translator that actually supports that filetype to handle the texture
+		UE::Interchange::Private::FScopedTranslator ScopedTranslator(TileFilePath, Results);
+		const IInterchangeTexturePayloadInterface* TextureTranslator = ScopedTranslator.GetPayLoadInterface<IInterchangeTexturePayloadInterface>();
+		if (!ensure(TextureTranslator))
+		{
+			continue;
+		}
+
+		// Invoke the translator to actually load the texture and parse it
+		const FString UnusedPayloadKey = {};
+		TOptional<UE::Interchange::FImportImage> TexturePayloadData;
+		TexturePayloadData = TextureTranslator->GetTexturePayloadData(UnusedPayloadKey, AlternateTexturePath);
+		if (!TexturePayloadData)
+		{
+			continue;
+		}
+		const UE::Interchange::FImportImage& Image = TileImages.Emplace_GetRef(MoveTemp(TexturePayloadData.GetValue()));
+		TexturePayloadData.Reset();
+
+		// Initialize the settings on the BlockData itself based on the first image we parse
+		if (!bInitializedBlockData)
+		{
+			bInitializedBlockData = true;
+
+			BlockData.Format = Image.Format;
+			BlockData.CompressionSettings = TextureGroup == TEXTUREGROUP_WorldNormalMap ? TC_Normalmap : TC_Default;
+			BlockData.bSRGB = Image.bSRGB;
+			BlockData.MipGenSettings = Image.MipGenSettings;
+		}
+
+		// Prepare the BlockData to receive this image data (later)
+		BlockData.InitBlockFromImage(BlockX, BlockY, Image);
+	}
+
+	// Move all of the FImportImage buffers into the BlockData itself
+	BlockData.MigrateDataFromImagesToRawData(TileImages);
+#endif	  // USE_USD_SDK
+
+	return BlockData;
+}
+
+TArray<UE::Interchange::FAnimationPayloadData> UInterchangeUSDTranslator::GetAnimationPayloadData(
+	const TArray<UE::Interchange::FAnimationPayloadQuery>& PayloadQueries
+) const
+{
+	using namespace UE::Interchange;
+	using namespace UE::InterchangeUsdTranslator::Private;
+	// This is the results we return
+	TArray<UE::Interchange::FAnimationPayloadData> AnimationPayloads;
+
+#if USE_USD_SDK
+	// Maps to help sorting the queries by payload type
+	TArray<int32> BakeQueryIndexes;
+	TArray<TArray<UE::Interchange::FAnimationPayloadData>> BakeAnimationPayloads;
+	TArray<int32> CurveQueryIndexes;
+	TArray<TArray<UE::Interchange::FAnimationPayloadData>> CurveAnimationPayloads;
+
+	// Get all curves with a parallel for
+	int32 PayloadCount = PayloadQueries.Num();
+	for (int32 PayloadIndex = 0; PayloadIndex < PayloadCount; ++PayloadIndex)
+	{
+		const UE::Interchange::FAnimationPayloadQuery& PayloadQuery = PayloadQueries[PayloadIndex];
+		EInterchangeAnimationPayLoadType QueryType = PayloadQuery.PayloadKey.Type;
+		if (QueryType == EInterchangeAnimationPayLoadType::BAKED)
+		{
+			BakeQueryIndexes.Add(PayloadIndex);
+		}
+		else
+		{
+			CurveQueryIndexes.Add(PayloadIndex);
+		}
+	}
+
+	// Import the Baked curve payloads
+	if (BakeQueryIndexes.Num() > 0)
+	{
+		int32 BakePayloadCount = BakeQueryIndexes.Num();
+		TMap<FString, TArray<const UE::Interchange::FAnimationPayloadQuery*>> BatchedBakeQueries;
+		BatchedBakeQueries.Reserve(BakePayloadCount);
+
+		// Get the BAKED transform synchronously, since there is some interchange task that parallel them
+		for (int32 BakePayloadIndex = 0; BakePayloadIndex < BakePayloadCount; ++BakePayloadIndex)
+		{
+			if (!ensure(BakeQueryIndexes.IsValidIndex(BakePayloadIndex)))
+			{
+				continue;
+			}
+			int32 PayloadIndex = BakeQueryIndexes[BakePayloadIndex];
+			if (!PayloadQueries.IsValidIndex(PayloadIndex))
+			{
+				continue;
+			}
+			const UE::Interchange::FAnimationPayloadQuery& PayloadQuery = PayloadQueries[PayloadIndex];
+			check(PayloadQuery.PayloadKey.Type == EInterchangeAnimationPayLoadType::BAKED);
+			// Joint transform animation queries.
+			//
+			// Currently we'll receive the PayloadQueries for all joints of a skeletal animation on the same GetAnimationPayloadData
+			// call. Unfortunately in USD we must compute all joint transforms every time, even if all we need is data for a single
+			// joint. For efficiency then, we group up all the queries for the separate joints of the same skeleton into one batch
+			// task that we can resolve in one pass
+			const FString BakedQueryHash = HashAnimPayloadQuery(PayloadQuery);
+			TArray<const UE::Interchange::FAnimationPayloadQuery*>& Queries = BatchedBakeQueries.FindOrAdd(BakedQueryHash);
+			Queries.Add(&PayloadQuery);
+		}
+		// Emit the batched joint transform animation tasks
+		for (const TPair<FString, TArray<const UE::Interchange::FAnimationPayloadQuery*>>& BatchedBakedQueryPair : BatchedBakeQueries)
+		{
+			const TArray<const UE::Interchange::FAnimationPayloadQuery*>& Queries = BatchedBakedQueryPair.Value;
+			TArray<UE::Interchange::FAnimationPayloadData> Result;
+			GetJointAnimationCurvePayloadData(*Impl, Queries, Result);
+			BakeAnimationPayloads.Add(Result);
+		}
+
+		// Append the bake curves results
+		for (TArray<UE::Interchange::FAnimationPayloadData>& AnimationPayload : BakeAnimationPayloads)
+		{
+			AnimationPayloads.Append(AnimationPayload);
+		}
+	}
+
+	// Import normal curves
+	if (CurveQueryIndexes.Num() > 0)
+	{
+		auto GetAnimPayloadLambda = [this, &PayloadQueries, &CurveAnimationPayloads](int32 PayloadIndex)
+		{
+			if (!PayloadQueries.IsValidIndex(PayloadIndex))
+			{
+				return;
+			}
+			const UE::Interchange::FAnimationPayloadQuery& PayloadQuery = PayloadQueries[PayloadIndex];
+			EInterchangeAnimationPayLoadType PayloadType = PayloadQuery.PayloadKey.Type;
+			if (PayloadType == EInterchangeAnimationPayLoadType::CURVE || PayloadType == EInterchangeAnimationPayLoadType::STEPCURVE)
+			{
+				// Property track animation queries.
+				//
+				// We're fine handling these in isolation (currently GetAnimationPayloadData is called with
+				// a single query at a time for these): Emit a separate task for each right away
+				FAnimationPayloadData AnimationPayLoadData{PayloadQuery.SceneNodeUniqueID, PayloadQuery.PayloadKey};
+				if (GetPropertyAnimationCurvePayloadData(*Impl, PayloadQuery.PayloadKey.UniqueId, AnimationPayLoadData))
+				{
+					CurveAnimationPayloads[PayloadIndex].Emplace(AnimationPayLoadData);
+				}
+			}
+			else if (PayloadType == EInterchangeAnimationPayLoadType::MORPHTARGETCURVE)
+			{
+				// Morph target curve queries.
+				FAnimationPayloadData AnimationPayLoadData{PayloadQuery.SceneNodeUniqueID, PayloadQuery.PayloadKey};
+				if (GetMorphTargetAnimationCurvePayloadData(*Impl, PayloadQuery.PayloadKey.UniqueId, AnimationPayLoadData))
+				{
+					CurveAnimationPayloads[PayloadIndex].Emplace(AnimationPayLoadData);
+				}
+			}
+		};
+
+		// Get all curves with a parallel for if there is many
+		int32 CurvePayloadCount = CurveQueryIndexes.Num();
+		CurveAnimationPayloads.AddDefaulted(CurvePayloadCount);
+		const int32 BatchSize = 10;
+		if (CurvePayloadCount > BatchSize)
+		{
+			const int32 NumBatches = (CurvePayloadCount / BatchSize) + 1;
+			ParallelFor(
+				NumBatches,
+				[&CurveQueryIndexes, &GetAnimPayloadLambda](int32 BatchIndex)
+				{
+					int32 PayloadIndexOffset = BatchIndex * BatchSize;
+					for (int32 PayloadIndex = PayloadIndexOffset; PayloadIndex < PayloadIndexOffset + BatchSize; ++PayloadIndex)
+					{
+						// The last batch can be incomplete
+						if (!CurveQueryIndexes.IsValidIndex(PayloadIndex))
+						{
+							break;
+						}
+						GetAnimPayloadLambda(CurveQueryIndexes[PayloadIndex]);
+					}
+				},
+				EParallelForFlags::BackgroundPriority	 // ParallelFor
+			);
+		}
+		else
+		{
+			for (int32 PayloadIndex = 0; PayloadIndex < CurvePayloadCount; ++PayloadIndex)
+			{
+				int32 PayloadQueriesIndex = CurveQueryIndexes[PayloadIndex];
+				if (PayloadQueries.IsValidIndex(PayloadQueriesIndex))
+				{
+					GetAnimPayloadLambda(PayloadQueriesIndex);
+				}
+			}
+		}
+
+		// Append the curves results
+		for (TArray<UE::Interchange::FAnimationPayloadData>& AnimationPayload : CurveAnimationPayloads)
+		{
+			AnimationPayloads.Append(AnimationPayload);
+		}
+	}
+#endif	  // USE_USD_SDK
+
+	return AnimationPayloads;
+}
+
+TOptional<UE::Interchange::FVolumePayloadData> UInterchangeUSDTranslator::GetVolumePayloadData(const UE::Interchange::FVolumePayloadKey& PayloadKey
+) const
+{
+#if USE_USD_SDK
+	TStrongObjectPtr<UInterchangeTranslatorBase> ExistingTranslator = Impl->Translators.FindRef(PayloadKey.FileName);
+	const IInterchangeVolumePayloadInterface* VolumeInterface = Cast<IInterchangeVolumePayloadInterface>(ExistingTranslator.Get());
+	if (VolumeInterface)
+	{
+		return VolumeInterface->GetVolumePayloadData(PayloadKey);
+	}
+#endif	  // USE_USD_SDK
+
+	return {};
+}
+
+#undef LOCTEXT_NAMESPACE
diff --git a/Engine/Plugins/Interchange/Extensions/OpenUSD/Source/Import/Public/InterchangeOpenUSDImportModule.h b/Engine/Plugins/Interchange/Extensions/OpenUSD/Source/Import/Public/InterchangeOpenUSDImportModule.h
new file mode 100644
--- /dev/null
+++ b/Engine/Plugins/Interchange/Extensions/OpenUSD/Source/Import/Public/InterchangeOpenUSDImportModule.h
@@ -0,0 +1,12 @@
+// Copyright Epic Games, Inc. All Rights Reserved.
+
+#pragma once
+
+#include "Modules/ModuleInterface.h"
+
+class FInterchangeOpenUSDImportModule : public IModuleInterface
+{
+public:
+	virtual void StartupModule() override;
+	virtual void ShutdownModule() override;
+};
\ No newline at end of file
diff --git a/Engine/Plugins/Interchange/Extensions/OpenUSD/Source/Import/Public/InterchangeUSDPipeline.h b/Engine/Plugins/Interchange/Extensions/OpenUSD/Source/Import/Public/InterchangeUSDPipeline.h
new file mode 100644
--- /dev/null
+++ b/Engine/Plugins/Interchange/Extensions/OpenUSD/Source/Import/Public/InterchangeUSDPipeline.h
@@ -0,0 +1,61 @@
+// Copyright Epic Games, Inc. All Rights Reserved.
+
+#pragma once
+
+#include "InterchangePipelineBase.h"
+#include "InterchangeSourceData.h"
+
+#include "InterchangeUSDPipeline.generated.h"
+
+class UInterchangeBaseNodeContainer;
+
+UCLASS(BlueprintType)
+class INTERCHANGEOPENUSDIMPORT_API UInterchangeUsdPipeline : public UInterchangePipelineBase
+{
+	GENERATED_BODY()
+
+	UInterchangeUsdPipeline();
+
+public:
+	UPROPERTY(
+		EditAnywhere,
+		BlueprintReadWrite,
+		Category = "Common",
+		meta = (StandAlonePipelineProperty = "True", PipelineInternalEditionData = "True")
+	)
+	FString PipelineDisplayName;
+
+	/** Only import translated nodes from imageable prims (Xforms, Meshes, etc.) with these specific purposes from the USD file */
+	UPROPERTY(EditAnywhere, BlueprintReadWrite, Category = "USD", meta = (Bitmask, BitmaskEnum = "/Script/UnrealUSDWrapper.EUsdPurpose"))
+	int32 GeometryPurpose;
+
+	/**
+	 * The translator always emits a scene node for the stage pseudoroot. Setting this option to false means we will remove it from
+	 * the factory nodes, making top-level prims lead to top-level actors and components.
+	 */
+	UPROPERTY(EditAnywhere, BlueprintReadWrite, Category = "USD")
+	bool bImportPseudoRoot;
+
+	/**
+	 * Imported meshes will place their primvars on arbitrary UV set indices (e.g. "st2" on UV0, "st3" on UV1, etc.). Imported materials may be
+	 * expecting to read from specific primvars on particular UV indices (e.g. "st1" from UV0, "st2" from UV2, etc.). These mappings won't
+	 * necessarily correspond to each other. Setting this option to true means we'll generate additional material instances with primvar-UV index
+	 * mappings made to best match the exact meshes they are assigned to (e.g. "st1" disabled, "st2" from UV0). Setting this to false means we
+	 * will just assign the incompatible materials to the meshes anyway.
+	 */
+	UPROPERTY(EditAnywhere, BlueprintReadWrite, Category = "USD")
+	bool bGeneratePrimvarCompatibleMaterials;
+
+public:
+
+	virtual void ExecutePostFactoryPipeline(const UInterchangeBaseNodeContainer* BaseNodeContainer, const FString& NodeKey, UObject* CreatedAsset, bool bIsAReimport) override;
+
+protected:
+	virtual void ExecutePipeline(
+		UInterchangeBaseNodeContainer* BaseNodeContainer,
+		const TArray<UInterchangeSourceData*>& SourceDatas,
+		const FString& ContentBasePath
+	) override;
+
+	virtual bool CanExecuteOnAnyThread(EInterchangePipelineTask PipelineTask) override;
+};
diff --git a/Engine/Plugins/Interchange/Extensions/OpenUSD/Source/Import/Public/InterchangeUsdContext.h b/Engine/Plugins/Interchange/Extensions/OpenUSD/Source/Import/Public/InterchangeUsdContext.h
new file mode 100644
--- /dev/null
+++ b/Engine/Plugins/Interchange/Extensions/OpenUSD/Source/Import/Public/InterchangeUsdContext.h
@@ -0,0 +1,57 @@
+// Copyright Epic Games, Inc. All Rights Reserved.
+
+#pragma once
+
+#include "UsdWrappers/UsdStage.h"
+
+#include "CoreMinimal.h"
+
+#include "InterchangeUsdContext.generated.h"
+
+class FUsdInfoCache;
+
+UCLASS(BlueprintType)
+class INTERCHANGEOPENUSDIMPORT_API UInterchangeUsdContext : public UObject
+{
+	GENERATED_BODY()
+
+public:
+	UInterchangeUsdContext();
+	virtual void BeginDestroy() override;
+
+	// Returns the ID the provided stage has within the USDUtils' singleton StageCache.
+	UFUNCTION(BlueprintCallable, Category = "Interchange | USD")
+	int64 GetStageId() const;
+
+	// Sets the ID of a particular stage from the UsdUtils' singleton StageCache.
+	// If this corresponds to a valid USD Stage, that stage will be used for the Interchange import.
+	UFUNCTION(BlueprintCallable, Category = "Interchange | USD")
+	void SetStageId(int64 InStageId);
+
+	// Convenience functions to get/set the stage directly, although it will internally just
+	// get/set the stage into the UsdUtils' singleton StageCache and track it's Id instead.
+	UE::FUsdStage GetUsdStage() const;
+	bool SetUsdStage(const UE::FUsdStage& InStage);
+
+	// Returns the current info cache, whether it is one we fully own, or just an external reference we're tracking
+	FUsdInfoCache* GetInfoCache() const;
+
+	// Receive a reference to an info cache that is external to this object and set it as the current info cache.
+	// Note that this will discard our owned info cache if we had one before.
+	void SetExternalInfoCache(FUsdInfoCache& InInfoCache);
+
+	// Create an info cache that is fully owned by this UInterchangeUsdContext object, and set it as the current info cache.
+	FUsdInfoCache* CreateOwnedInfoCache();
+
+	// Properly releases the reference to the current stage cache, whether owned or not (deleting it if owned)
+	void ReleaseInfoCache();
+
+private:
+	// We never store the stage itself, but only it's Id within the UsdUtils' singleton StageCache.
+	// The intent here is to allow Python stages to be passed in and manipulated via Python in case of Python imports or pipelines.
+	int64 StageIdInUsdUtilsStageCache = INDEX_NONE;
+	bool bShouldCleanUpFromStageCache = false;
+
+	FUsdInfoCache* InfoCache = nullptr;
+	TUniquePtr<FUsdInfoCache> OwnedInfoCache;
+};
diff --git a/Engine/Plugins/Interchange/Extensions/OpenUSD/Source/Import/Public/InterchangeUsdTranslator.h b/Engine/Plugins/Interchange/Extensions/OpenUSD/Source/Import/Public/InterchangeUsdTranslator.h
new file mode 100644
--- /dev/null
+++ b/Engine/Plugins/Interchange/Extensions/OpenUSD/Source/Import/Public/InterchangeUsdTranslator.h
@@ -0,0 +1,125 @@
+// Copyright Epic Games, Inc. All Rights Reserved.
+
+#pragma once
+
+#include "Animation/InterchangeAnimationPayloadInterface.h"
+#include "InterchangeTranslatorBase.h"
+#include "Mesh/InterchangeMeshPayloadInterface.h"
+#include "Texture/InterchangeBlockedTexturePayloadInterface.h"
+#include "Texture/InterchangeTexturePayloadInterface.h"
+#include "Volume/InterchangeVolumePayloadInterface.h"
+
+#include "USDStageOptions.h"
+
+#include "UObject/Object.h"
+#include "UObject/ObjectMacros.h"
+
+#include "InterchangeUsdTranslator.generated.h"
+
+enum class EUsdInterpolationType : uint8;
+namespace UE::InterchangeUsdTranslator::Private
+{
+	class UInterchangeUSDTranslatorImpl;
+}
+
+UCLASS(BlueprintType, editinlinenew, MinimalAPI)
+class UInterchangeUsdTranslatorSettings : public UInterchangeTranslatorSettings
+{
+	GENERATED_BODY()
+
+public:
+
+	/** Only import geometry prims with these specific purposes from the USD file */
+	UE_DEPRECATED(5.6, "This option is no longer a translator setting and has moved to the InterchangeUSDPipeline")
+	UPROPERTY()
+	int32 GeometryPurpose;
+
+	/** Specifies which set of shaders to use when parsing USD materials, in addition to the universal render context. */
+	UPROPERTY(EditAnywhere, Category = "USD Translator")
+	FName RenderContext;
+
+	/** Specifies which material purpose to use when parsing USD material bindings, in addition to the "allPurpose" fallback */
+	UPROPERTY(EditAnywhere, Category = "USD Translator")
+	FName MaterialPurpose;
+
+	/** Describes how to interpolate between a timeSample value and the next */
+	UPROPERTY(EditAnywhere, Category = "USD Translator")
+	EUsdInterpolationType InterpolationType;
+
+	/** Whether to use the specified StageOptions instead of the stage's own settings */
+	UPROPERTY(EditAnywhere, Category = "USD Translator")
+	bool bOverrideStageOptions;
+
+	/** Custom StageOptions to use for the stage */
+	UPROPERTY(EditAnywhere, Category = "USD Translator", meta = (EditCondition = bOverrideStageOptions))
+	FUsdStageOptions StageOptions;
+
+	/** Whether to convert USD prim attributes into translated node custom attributes, if they pass the regex filter */
+	UPROPERTY(EditAnywhere, Category = "USD Translator")
+	bool bTranslatePrimAttributes;
+
+	/** Regex filter to select which USD prim attributes should be converted into translated node custom attributes */
+	UPROPERTY(EditAnywhere, Category = "USD Translator", meta = (EditCondition = bTranslatePrimAttributes))
+	FString AttributeRegexFilter;
+
+public:
+	UInterchangeUsdTranslatorSettings();
+};
+
+/* For now, USD Interchange (FBX parity) translator supports textures, materials and static meshes */
+UCLASS(BlueprintType)
+class UInterchangeUSDTranslator
+	: public UInterchangeTranslatorBase
+	, public IInterchangeMeshPayloadInterface
+	, public IInterchangeTexturePayloadInterface
+	, public IInterchangeBlockedTexturePayloadInterface
+	, public IInterchangeAnimationPayloadInterface
+	, public IInterchangeVolumePayloadInterface
+{
+	GENERATED_BODY()
+
+public:
+	UInterchangeUSDTranslator();
+
+	/** Begin UInterchangeTranslatorBase API*/
+	virtual EInterchangeTranslatorType GetTranslatorType() const override;
+	virtual EInterchangeTranslatorAssetType GetSupportedAssetTypes() const override;
+	virtual TArray<FString> GetSupportedFormats() const override;
+	virtual bool Translate(UInterchangeBaseNodeContainer& BaseNodeContainer) const override;
+	virtual void ReleaseSource() override;
+	virtual void ImportFinish() override;
+	virtual UInterchangeTranslatorSettings* GetSettings() const override;
+	virtual void SetSettings(const UInterchangeTranslatorSettings* InterchangeTranslatorSettings) override;
+	/** End UInterchangeTranslatorBase API*/
+
+	TFuture<TOptional<UE::Interchange::FAnimationPayloadData>> ResolveAnimationPayloadQuery(
+		const UE::Interchange::FAnimationPayloadQuery& PayloadQuery
+	) const;
+
+	/** Begin Interchange payload interfaces */
+	virtual TOptional<UE::Interchange::FMeshPayloadData> GetMeshPayloadData(
+		const FInterchangeMeshPayLoadKey& PayLoadKey,
+		const FTransform& MeshGlobalTransform
+	) const override;
+
+	virtual TOptional<UE::Interchange::FImportImage> GetTexturePayloadData(const FString& PayloadKey, TOptional<FString>& AlternateTexturePath)
+		const override;
+
+	virtual TOptional<UE::Interchange::FImportBlockedImage> GetBlockedTexturePayloadData(
+		const FString& PayloadKey,
+		TOptional<FString>& AlternateTexturePath
+	) const override;
+
+	virtual TArray<UE::Interchange::FAnimationPayloadData> GetAnimationPayloadData(
+		const TArray<UE::Interchange::FAnimationPayloadQuery>& PayloadQueries
+	) const override;
+
+	virtual TOptional<UE::Interchange::FVolumePayloadData> GetVolumePayloadData(const UE::Interchange::FVolumePayloadKey& PayloadKey) const override;
+	/** End Interchange payload interfaces */
+
+private:
+	mutable TUniquePtr<UE::InterchangeUsdTranslator::Private::UInterchangeUSDTranslatorImpl> Impl;
+
+	UPROPERTY(Transient, DuplicateTransient)
+	mutable TObjectPtr<UInterchangeUsdTranslatorSettings> TranslatorSettings = nullptr;
+};
diff --git a/Engine/Plugins/Interchange/Runtime/Interchange.uplugin b/Engine/Plugins/Interchange/Runtime/Interchange.uplugin
--- a/Engine/Plugins/Interchange/Runtime/Interchange.uplugin
+++ b/Engine/Plugins/Interchange/Runtime/Interchange.uplugin
@@ -92,7 +92,7 @@
 			"Enabled": true
 		},
 		{
-			"Name": "USDCore",
+			"Name": "GeometryProcessing",
 			"Enabled": true
 		},
 		{
@@ -106,11 +106,6 @@
 		{
 			"Name": "InterchangeAssets",
 			"Enabled": true
-		},
-		{
-			"Name": "MeshModelingToolset",
-			"Enabled": true,
-			"TargetAllowList": [ "Editor" ]
 		}
 	],
 	"IsExperimentalVersion": false,
diff --git a/Engine/Plugins/Interchange/Runtime/Source/Import/InterchangeImport.Build.cs b/Engine/Plugins/Interchange/Runtime/Source/Import/InterchangeImport.Build.cs
--- a/Engine/Plugins/Interchange/Runtime/Source/Import/InterchangeImport.Build.cs
+++ b/Engine/Plugins/Interchange/Runtime/Source/Import/InterchangeImport.Build.cs
@@ -44,9 +44,6 @@ namespace UnrealBuildTool.Rules
 					"RenderCore",
 					"RHI",
 					"TextureUtilitiesCommon",
-					"UnrealUSDWrapper",
-					"USDClasses",
-					"USDUtilities",
 					"VariantManagerContent",
 				}
 			);
@@ -68,7 +65,7 @@ namespace UnrealBuildTool.Rules
 					}
 				);
 
-				if (Target.Platform == UnrealTargetPlatform.Win64 || Target.LinkType == TargetLinkType.Monolithic)
+				if (Target.Platform == UnrealTargetPlatform.Win64 || Target.Platform == UnrealTargetPlatform.Mac || Target.LinkType == TargetLinkType.Monolithic)
 				{
 					PrivateDependencyModuleNames.AddRange(
 						new string[]
diff --git a/Engine/Plugins/Interchange/Runtime/Source/Import/Private/InterchangeImportModule.cpp b/Engine/Plugins/Interchange/Runtime/Source/Import/Private/InterchangeImportModule.cpp
--- a/Engine/Plugins/Interchange/Runtime/Source/Import/Private/InterchangeImportModule.cpp
+++ b/Engine/Plugins/Interchange/Runtime/Source/Import/Private/InterchangeImportModule.cpp
@@ -40,7 +40,6 @@
 #include "Texture/InterchangePSDTranslator.h"
 #include "Texture/InterchangeTextureFactory.h"
 #include "Texture/InterchangeUEJPEGTranslator.h"
-#include "Usd/InterchangeUsdTranslator.h"
 #include "Volume/InterchangeSparseVolumeTextureFactory.h"
 
 DEFINE_LOG_CATEGORY(LogInterchangeImport);
@@ -87,7 +86,6 @@ void FInterchangeImportModule::StartupModule()
 		UInterchangeFbxTranslator::CleanUpTemporaryFolder();
 		InterchangeManager.RegisterTranslator(UInterchangeFbxTranslator::StaticClass());
 #endif
-		InterchangeManager.RegisterTranslator(UInterchangeUSDTranslator::StaticClass());
 		InterchangeManager.RegisterTranslator(UInterchangeGLTFTranslator::StaticClass());
 		InterchangeManager.RegisterTranslator(UInterchangeOBJTranslator::StaticClass());
 
diff --git a/Engine/Plugins/Interchange/Runtime/Source/Import/Private/Usd/InterchangeUsdContext.cpp b/Engine/Plugins/Interchange/Runtime/Source/Import/Private/Usd/InterchangeUsdContext.cpp
deleted file mode 100644
diff --git a/Engine/Plugins/Interchange/Runtime/Source/Import/Private/Usd/InterchangeUsdTranslator.cpp b/Engine/Plugins/Interchange/Runtime/Source/Import/Private/Usd/InterchangeUsdTranslator.cpp
deleted file mode 100644
diff --git a/Engine/Plugins/Interchange/Runtime/Source/Import/Public/MaterialX/MaterialXUtils/MaterialXManager.h b/Engine/Plugins/Interchange/Runtime/Source/Import/Public/MaterialX/MaterialXUtils/MaterialXManager.h
--- a/Engine/Plugins/Interchange/Runtime/Source/Import/Public/MaterialX/MaterialXUtils/MaterialXManager.h
+++ b/Engine/Plugins/Interchange/Runtime/Source/Import/Public/MaterialX/MaterialXUtils/MaterialXManager.h
@@ -16,7 +16,7 @@ class UInterchangeBaseNodeContainer;
 class UInterchangeShaderNode;
 class FMaterialXSurfaceShaderAbstract;
 
-class FMaterialXManager
+class INTERCHANGEIMPORT_API FMaterialXManager
 {
 public:
 
diff --git a/Engine/Plugins/Interchange/Runtime/Source/Import/Public/Texture/InterchangeBlockedTexturePayloadData.h b/Engine/Plugins/Interchange/Runtime/Source/Import/Public/Texture/InterchangeBlockedTexturePayloadData.h
--- a/Engine/Plugins/Interchange/Runtime/Source/Import/Public/Texture/InterchangeBlockedTexturePayloadData.h
+++ b/Engine/Plugins/Interchange/Runtime/Source/Import/Public/Texture/InterchangeBlockedTexturePayloadData.h
@@ -13,7 +13,7 @@ namespace UE
 	namespace Interchange
 	{
 		// Also known as a UDIMs texture
-		struct FImportBlockedImage
+		struct INTERCHANGEIMPORT_API FImportBlockedImage
 		{
 			FUniqueBuffer RawData;
 			TArray<FTextureSourceBlock> BlocksData;
diff --git a/Engine/Plugins/Interchange/Runtime/Source/Import/Public/Usd/InterchangeUsdContext.h b/Engine/Plugins/Interchange/Runtime/Source/Import/Public/Usd/InterchangeUsdContext.h
deleted file mode 100644
diff --git a/Engine/Plugins/Interchange/Runtime/Source/Import/Public/Usd/InterchangeUsdTranslator.h b/Engine/Plugins/Interchange/Runtime/Source/Import/Public/Usd/InterchangeUsdTranslator.h
deleted file mode 100644
diff --git a/Engine/Plugins/Interchange/Runtime/Source/Pipelines/InterchangePipelines.Build.cs b/Engine/Plugins/Interchange/Runtime/Source/Pipelines/InterchangePipelines.Build.cs
--- a/Engine/Plugins/Interchange/Runtime/Source/Pipelines/InterchangePipelines.Build.cs
+++ b/Engine/Plugins/Interchange/Runtime/Source/Pipelines/InterchangePipelines.Build.cs
@@ -26,22 +26,13 @@ namespace UnrealBuildTool.Rules
 					"ApplicationCore",
 					"CinematicCamera",
 					"DeveloperSettings",
-					"DynamicMesh",
 					"GeometryCache",
-					"GeometryCore",
 					"GLTFCore",
 					"InputCore",
 					"InterchangeImport",
 					"LevelSequence",
-					"MeshConversion",
-					"MeshDescription",
-					"MeshUtilitiesCommon",
 					"Slate",
 					"SlateCore",
-					"StaticMeshDescription",
-					"UnrealUSDWrapper",
-					"USDClasses",
-					"USDUtilities",
 				}
 			);
 
@@ -51,7 +42,6 @@ namespace UnrealBuildTool.Rules
 					new string[]
 					{
 						"MainFrame",
-						"ModelingComponents",
 						"PhysicsUtilities",
 						"TextureUtilitiesCommon",
 						"UnrealEd",
diff --git a/Engine/Plugins/Interchange/Runtime/Source/Pipelines/Private/InterchangeUSDPipeline.cpp b/Engine/Plugins/Interchange/Runtime/Source/Pipelines/Private/InterchangeUSDPipeline.cpp
deleted file mode 100644
diff --git a/Engine/Plugins/Interchange/Runtime/Source/Pipelines/Public/InterchangeUSDPipeline.h b/Engine/Plugins/Interchange/Runtime/Source/Pipelines/Public/InterchangeUSDPipeline.h
deleted file mode 100644
