Interchange USD: Fixed Mac build issue

#rnx

diff --git a/Engine/Plugins/Interchange/Extensions/OpenUSD/Source/Import/InterchangeOpenUSDImport.build.cs b/Engine/Plugins/Interchange/Extensions/OpenUSD/Source/Import/InterchangeOpenUSDImport.build.cs
--- a/Engine/Plugins/Interchange/Extensions/OpenUSD/Source/Import/InterchangeOpenUSDImport.build.cs
+++ b/Engine/Plugins/Interchange/Extensions/OpenUSD/Source/Import/InterchangeOpenUSDImport.build.cs
@@ -45,15 +45,10 @@ namespace UnrealBuildTool.Rules
 					new string[]
 					{
 						"InterchangeOpenVDBImport",
+						"MaterialX",
 						"ModelingComponents",	// For FTexture2DBuilder
 					}
 				);
-
-				if (Target.Platform == UnrealTargetPlatform.Win64 || Target.LinkType == TargetLinkType.Monolithic ||
-					(Target.Platform == UnrealTargetPlatform.Mac && Target.Architecture == UnrealArch.X64))
-				{
-					PrivateDependencyModuleNames.Add("MaterialX");
-				}
 			}
 		}
 	}
