USD: Instrument InterchangeUsdTranslator for Unreal Insights.

diff --git a/Engine/Plugins/Interchange/Extensions/OpenUSD/Source/Import/Private/InterchangeUsdTranslator.cpp b/Engine/Plugins/Interchange/Extensions/OpenUSD/Source/Import/Private/InterchangeUsdTranslator.cpp
--- a/Engine/Plugins/Interchange/Extensions/OpenUSD/Source/Import/Private/InterchangeUsdTranslator.cpp
+++ b/Engine/Plugins/Interchange/Extensions/OpenUSD/Source/Import/Private/InterchangeUsdTranslator.cpp
@@ -1258,6 +1258,8 @@ namespace UE::InterchangeUsdTranslator::Private
 		UInterchangeBaseNodeContainer& NodeContainer
 	)
 	{
+		TRACE_CPUPROFILER_EVENT_SCOPE(UInterchangeUSDTranslatorImpl::AddMaterialNode)
+
 		FName RenderContext = TranslatorSettings ? TranslatorSettings->RenderContext : UnrealIdentifiers::UniversalRenderContext;
 
 		// If this material has an unreal surface output and we're in the unreal render context, just emit a material reference,
@@ -1553,6 +1555,8 @@ namespace UE::InterchangeUsdTranslator::Private
 		const UInterchangeSceneNode* SceneNodeWithLODs
 	)
 	{
+		TRACE_CPUPROFILER_EVENT_SCOPE(UInterchangeUSDTranslatorImpl::FinalizeLODContainerTraversal)
+
 		if (!SceneNodeWithLODs)
 		{
 			return;
@@ -1664,6 +1668,8 @@ namespace UE::InterchangeUsdTranslator::Private
 
 	FString AddLightNode(const UE::FUsdPrim& Prim, UInterchangeBaseNodeContainer& NodeContainer)
 	{
+		TRACE_CPUPROFILER_EVENT_SCOPE(AddLightNode)
+
 		FString NodeUid;
 
 		NodeUid = LightPrefix + Prim.GetPrimPath().GetString();
@@ -1801,6 +1807,8 @@ namespace UE::InterchangeUsdTranslator::Private
 
 	FString AddCameraNode(const UE::FUsdPrim& Prim, UInterchangeBaseNodeContainer& NodeContainer)
 	{
+		TRACE_CPUPROFILER_EVENT_SCOPE(AddCameraNode)
+
 		FString NodeUid;
 
 		NodeUid = CameraPrefix + Prim.GetPrimPath().GetString();
@@ -1904,6 +1912,8 @@ namespace UE::InterchangeUsdTranslator::Private
 		bool bPrimitiveShape
 	)
 	{
+		TRACE_CPUPROFILER_EVENT_SCOPE(UInterchangeUSDTranslatorImpl::AddMeshNode)
+
 		FString PrimPath = Prim.GetPrimPath().GetString();
 		FString NodeUid = (bPrimitiveShape ? PrimitiveShapePrefix : MeshPrefix) + PrimPath;
 		FString NodeName(Prim.GetName().ToString());
@@ -2102,6 +2112,8 @@ namespace UE::InterchangeUsdTranslator::Private
 		FTraversalInfo Info
 	)
 	{
+		TRACE_CPUPROFILER_EVENT_SCOPE(UInterchangeUSDTranslatorImpl::AddLODMeshNodes)
+
 		UE::FUsdVariantSets VariantSets = Prim.GetVariantSets();
 		const UE::FUsdVariantSet LODVariantSet = VariantSets.GetVariantSet(LODString);
 		if (!LODVariantSet.IsValid())
@@ -2320,6 +2332,8 @@ namespace UE::InterchangeUsdTranslator::Private
 		bool& bOutNeedsFrameTrack
 	)
 	{
+		TRACE_CPUPROFILER_EVENT_SCOPE(UInterchangeUSDTranslatorImpl::AddVolumeNodes)
+
 		using namespace UE::Interchange::USD;
 
 		// Each Volume prim can reference multiple FieldAsset prims. Each FieldAsset itself can point to a particular grid within
@@ -2650,6 +2664,8 @@ namespace UE::InterchangeUsdTranslator::Private
 
 	void AddTransformAnimationNode(const UE::FUsdPrim& Prim, UInterchangeUSDTranslatorImpl& Impl, UInterchangeBaseNodeContainer& NodeContainer)
 	{
+		TRACE_CPUPROFILER_EVENT_SCOPE(AddTransformAnimationNode)
+
 		const FString PrimPath = Prim.GetPrimPath().GetString();
 		const FString UniquePath = PrimPath + TEXT("\\") + UnrealIdentifiers::TransformPropertyName.ToString();
 		const FString AnimTrackNodeUid = AnimationTrackPrefix + UniquePath;
@@ -2709,6 +2725,8 @@ namespace UE::InterchangeUsdTranslator::Private
 
 	void AddPropertyAnimationNodes(const UE::FUsdPrim& Prim, UInterchangeUSDTranslatorImpl& Impl, UInterchangeBaseNodeContainer& NodeContainer)
 	{
+		TRACE_CPUPROFILER_EVENT_SCOPE(AddPropertyAnimationNodes)
+
 		using namespace UE::InterchangeUsdTranslator::Private;
 
 		if (!Prim)
@@ -2760,6 +2778,8 @@ namespace UE::InterchangeUsdTranslator::Private
 		UInterchangeBaseNodeContainer& NodeContainer
 	)
 	{
+		TRACE_CPUPROFILER_EVENT_SCOPE(AddVolumeSceneNodeAttributes)
+
 		if (!SceneNode)
 		{
 			return;
@@ -3018,6 +3038,8 @@ namespace UE::InterchangeUsdTranslator::Private
 		FTraversalInfo& Info
 	)
 	{
+		TRACE_CPUPROFILER_EVENT_SCOPE(AddSkeletonNodes)
+
 		// If we're not inside of a SkelRoot, the skeleton shouldn't really do anything
 		if (!Info.ClosestParentSkelRootPath.IsValid())
 		{
@@ -3153,6 +3175,8 @@ namespace UE::InterchangeUsdTranslator::Private
 
 	void TranslateAttributes(const UE::FUsdPrim& Prim, UInterchangeSceneNode* Node, const FString& AllowedAttributeRegex)
 	{
+		TRACE_CPUPROFILER_EVENT_SCOPE(TranslateAttributes)
+
 		FRegexPattern RegexPattern{AllowedAttributeRegex};
 
 		for (const UE::FUsdAttribute& Attr : Prim.GetAttributes())
@@ -3256,6 +3280,8 @@ namespace UE::InterchangeUsdTranslator::Private
 		FTraversalInfo Info
 	)
 	{
+		TRACE_CPUPROFILER_EVENT_SCOPE(Traverse)
+
 		FString SceneNodeUid = Prim.GetPrimPath().GetString();
 		FString DisplayLabel(Prim.GetName().ToString());
 
@@ -3431,6 +3457,8 @@ namespace UE::InterchangeUsdTranslator::Private
 		FMeshDescription& OutMeshDescription
 	)
 	{
+		TRACE_CPUPROFILER_EVENT_SCOPE(GetStaticMeshPayloadData)
+
 		const bool bIsLODMesh = CheckAndChopPayloadPrefix(PayloadKey, LODPrefix);
 		const bool bIsPrimitiveShape = CheckAndChopPayloadPrefix(PayloadKey, PrimitiveShapePrefix);
 
@@ -3504,6 +3532,8 @@ namespace UE::InterchangeUsdTranslator::Private
 		TArray<FString>& OutJointNames
 	)
 	{
+		TRACE_CPUPROFILER_EVENT_SCOPE(GetSkeletalMeshPayloadData)
+
 #if WITH_EDITOR
 		const bool bIsLODMesh = CheckAndChopPayloadPrefix(PayloadKey, LODPrefix);
 
@@ -3621,6 +3651,8 @@ namespace UE::InterchangeUsdTranslator::Private
 		FString& OutMorphTargetName
 	)
 	{
+		TRACE_CPUPROFILER_EVENT_SCOPE(GetMorphTargetPayloadData)
+
 		bool bIsLODMesh = false;
 		FString MeshPrimPath;
 		int32 BlendShapeIndex = INDEX_NONE;
@@ -3684,6 +3716,8 @@ namespace UE::InterchangeUsdTranslator::Private
 		TArray<double>& OutTimeSampleUnion
 	)
 	{
+		TRACE_CPUPROFILER_EVENT_SCOPE(CreateVolumeTrackReader)
+
 		UsdToUnreal::FPropertyTrackReader Result;
 
 		const TArray<UsdUtils::FVolumePrimInfo>* CollectedInfoForPrim = InImpl.PrimPathToVolumeInfo.Find(InPrimPath);
@@ -3751,6 +3785,8 @@ namespace UE::InterchangeUsdTranslator::Private
 		Interchange::FAnimationPayloadData& OutPayloadData
 	)
 	{
+		TRACE_CPUPROFILER_EVENT_SCOPE(GetPropertyAnimationCurvePayloadData)
+
 		FString PrimPath;
 		FString UEPropertyNameStr;
 		bool bSplit = PayloadKey.Split(TEXT("\\"), &PrimPath, &UEPropertyNameStr, ESearchCase::CaseSensitive, ESearchDir::FromEnd);
@@ -3812,6 +3848,8 @@ namespace UE::InterchangeUsdTranslator::Private
 		TArray<UE::Interchange::FAnimationPayloadData>& OutPayloadData
 	)
 	{
+		TRACE_CPUPROFILER_EVENT_SCOPE(GetJointAnimationCurvePayloadData)
+
 		if (Queries.Num() == 0)
 		{
 			return false;
@@ -3935,6 +3973,8 @@ namespace UE::InterchangeUsdTranslator::Private
 		Interchange::FAnimationPayloadData& OutPayloadData
 	)
 	{
+		TRACE_CPUPROFILER_EVENT_SCOPE(GetMorphTargetAnimationCurvePayloadData)
+
 		// Here we must output the morph target curve for a particular channel and skinning target, i.e.
 		// the connection of a SkelAnimation blend shape channel to a particular Mesh prim.
 
@@ -4157,6 +4197,8 @@ TArray<FString> UInterchangeUSDTranslator::GetSupportedFormats() const
 bool UInterchangeUSDTranslator::Translate(UInterchangeBaseNodeContainer& NodeContainer) const
 {
 #if USE_USD_SDK
+	TRACE_CPUPROFILER_EVENT_SCOPE(UInterchangeUSDTranslator::Translate)
+
 	using namespace UE;
 	using namespace UE::InterchangeUsdTranslator::Private;
 	using namespace UsdToUnreal;
@@ -4379,6 +4421,8 @@ TOptional<UE::Interchange::FMeshPayloadData> UInterchangeUSDTranslator::GetMeshP
 	const FTransform& MeshGlobalTransform
 ) const
 {
+	TRACE_CPUPROFILER_EVENT_SCOPE(UInterchangeUSDTranslator::GetMeshPayloadData)
+
 	using namespace UE::InterchangeUsdTranslator::Private;
 	bool bSuccess = false;
 	TOptional<UE::Interchange::FMeshPayloadData> Result;
@@ -4464,6 +4508,8 @@ namespace UE::InterchangeUsdTranslator::Private
 		UInterchangeUSDTranslatorImpl& Impl
 	)
 	{
+		TRACE_CPUPROFILER_EVENT_SCOPE(GetMaterialXTexturePayloadData)
+
 		TOptional<UE::Interchange::FImportImage> TexturePayloadData;
 
 #if USE_USD_SDK
@@ -4514,6 +4560,8 @@ TOptional<UE::Interchange::FImportImage> UInterchangeUSDTranslator::GetTexturePa
 	TOptional<FString>& AlternateTexturePath
 ) const
 {
+	TRACE_CPUPROFILER_EVENT_SCOPE(UInterchangeUSDTranslator::GetTexturePayloadData)
+
 	using namespace UE::InterchangeUsdTranslator::Private;
 
 	TOptional<UE::Interchange::FImportImage> TexturePayloadData;
@@ -4584,6 +4632,8 @@ TOptional<UE::Interchange::FImportBlockedImage> UInterchangeUSDTranslator::GetBl
 	TOptional<FString>& AlternateTexturePath
 ) const
 {
+	TRACE_CPUPROFILER_EVENT_SCOPE(UInterchangeUSDTranslator::GetBlockedTexturePayloadData)
+
 	using namespace UE::InterchangeUsdTranslator::Private;
 
 	UE::Interchange::FImportBlockedImage BlockData;
@@ -4679,6 +4729,8 @@ TArray<UE::Interchange::FAnimationPayloadData> UInterchangeUSDTranslator::GetAni
 	const TArray<UE::Interchange::FAnimationPayloadQuery>& PayloadQueries
 ) const
 {
+	TRACE_CPUPROFILER_EVENT_SCOPE(UInterchangeUSDTranslator::GetAnimationPayloadData)
+
 	using namespace UE::Interchange;
 	using namespace UE::InterchangeUsdTranslator::Private;
 	// This is the results we return
@@ -4839,6 +4891,8 @@ TArray<UE::Interchange::FAnimationPayloadData> UInterchangeUSDTranslator::GetAni
 TOptional<UE::Interchange::FVolumePayloadData> UInterchangeUSDTranslator::GetVolumePayloadData(const UE::Interchange::FVolumePayloadKey& PayloadKey
 ) const
 {
+	TRACE_CPUPROFILER_EVENT_SCOPE(UInterchangeUSDTranslator::GetVolumePayloadData)
+
 #if USE_USD_SDK
 	TStrongObjectPtr<UInterchangeTranslatorBase> ExistingTranslator = Impl->Translators.FindRef(PayloadKey.FileName);
 	const IInterchangeVolumePayloadInterface* VolumeInterface = Cast<IInterchangeVolumePayloadInterface>(ExistingTranslator.Get());
