USD: Fix "Import as Material" vs "Import as Material Instance" issues for USD Interchange.

Hide "Material Import" option on default USD pipelines, and force it to "Import as Material";
Always emit material instance nodes from UsdPreviewSurface / displayColor materials;
Always handle USD MaterialX materials as UMaterials for now;
Fix displayColor / MaterialX doubleSided materials not working;
Fix some material types not working if the material prim was visited before the mesh prim during traversal;
Refactor / simplify handling of USD MaterialX materials.

#virtualized

diff --git a/Engine/Plugins/Interchange/Extensions/OpenUSD/Source/Import/Private/InterchangeUSDPipeline.cpp b/Engine/Plugins/Interchange/Extensions/OpenUSD/Source/Import/Private/InterchangeUSDPipeline.cpp
--- a/Engine/Plugins/Interchange/Extensions/OpenUSD/Source/Import/Private/InterchangeUSDPipeline.cpp
+++ b/Engine/Plugins/Interchange/Extensions/OpenUSD/Source/Import/Private/InterchangeUSDPipeline.cpp
@@ -6,6 +6,7 @@
 #include "InterchangeManager.h"
 #include "InterchangeMaterialDefinitions.h"
 #include "InterchangeMaterialFactoryNode.h"
+#include "InterchangeMaterialInstanceNode.h"
 #include "InterchangeMeshFactoryNode.h"
 #include "InterchangeMeshNode.h"
 #include "InterchangeSceneNode.h"
@@ -211,240 +212,64 @@ namespace UE::InterchangeUsdPipeline::Private
 		);
 	}
 
-	// Resolves which should be the parent material for NewFactoryNode and sets it
-	void SetCustomParentMaterial(const UInterchangeShaderGraphNode& ShaderGraphNode, UInterchangeMaterialInstanceFactoryNode& NewFactoryNode)
+	// Moves our primvar-compatible custom attributes from our flagged material instance translated nodes to the
+	// corresponding factory nodes.
+	//
+	// Must run before ProcessMeshNodes(), as it will use these attributes
+	void MoveCustomAttributesToMaterialFactoryNodes(UInterchangeBaseNodeContainer& NodeContainer)
 	{
-		using namespace UsdUnreal::MaterialUtils;
-
-		// Priority 1) Explicitly set parent material
-		//
-		// This can be from an UnrealMaterial or from anything the user may have manually set
-		// TODO: Maybe use this to handle UnrealMaterials?
-		if (FString ParentMaterialPath;
-			ShaderGraphNode.GetStringAttribute(UE::Interchange::USD::ParentMaterialIdentifier, ParentMaterialPath) && !ParentMaterialPath.IsEmpty())
-		{
-			NewFactoryNode.SetCustomParent(ParentMaterialPath);
-			return;
-		}
-
-		// Priority 2) DisplayColor material
-		//
-		// This is usually a lower priority, but we explicitly flag these so it's easy to check for them
-		if (FString DisplayColorDesc;
-			ShaderGraphNode.GetStringAttribute(UE::Interchange::USD::DisplayColorDescIdentifier, DisplayColorDesc) && !DisplayColorDesc.IsEmpty())
-		{
-			TOptional<FDisplayColorMaterial> ParsedMat = FDisplayColorMaterial::FromString(DisplayColorDesc);
-			if (ParsedMat.IsSet())
-			{
-				if (const FSoftObjectPath* ParentMaterialPath = GetReferenceMaterialPath(ParsedMat.GetValue()))
-				{
-					NewFactoryNode.SetCustomParent(ParentMaterialPath->GetAssetPathString());
-					return;
-				}
-			}
-		}
-
-		// Priority 3) UsdPreviewSurface material
-		{
-			EUsdReferenceMaterialProperties Properties = EUsdReferenceMaterialProperties::None;
-
-			// Translucent
-			int BlendMode = EBlendMode::BLEND_Opaque;
-			if (ShaderGraphNode.GetCustomBlendMode(BlendMode) && BlendMode == EBlendMode::BLEND_Translucent)
-			{
-				Properties |= EUsdReferenceMaterialProperties::Translucent;
-			}
-
-			// TwoSided
-			// TODO: Actually set this somewhere
-			bool bTwoSided = false;
-			if (ShaderGraphNode.GetCustomTwoSided(bTwoSided) && bTwoSided)
-			{
-				Properties |= EUsdReferenceMaterialProperties::TwoSided;
-			}
-
-			// VT
-			bool bNeedsVT = false;
-			if (ShaderGraphNode.GetBooleanAttribute(UE::Interchange::USD::NeedsVTParentIdentifier, bNeedsVT) && bNeedsVT)
-			{
-				Properties |= EUsdReferenceMaterialProperties::VT;
-			}
-
-			FSoftObjectPath ParentMaterial = GetReferencePreviewSurfaceMaterial(Properties);
-			if (ParentMaterial.IsValid())
+		NodeContainer.IterateNodesOfType<UInterchangeMaterialInstanceFactoryNode>(
+			[&NodeContainer](const FString& NodeUid, UInterchangeMaterialInstanceFactoryNode* FactoryNode)
 			{
-				NewFactoryNode.SetCustomParent(ParentMaterial.GetAssetPathString());
-				return;
-			}
-		}
-	}
-
-	// Creates a new material instance factory node and fills it in with the parameter/parent material path
-	// stored as base translated node attributes on the provided ShaderGraphNode
-	UInterchangeMaterialInstanceFactoryNode* BuildMaterialInstance(
-		UInterchangeBaseNodeContainer& NodeContainer,
-		const UInterchangeShaderGraphNode& ShaderGraphNode
-	)
-	{
-		const FString ShaderGraphNodeUid = ShaderGraphNode.GetUniqueID();
-		FString MaterialFactoryNodeUid = UInterchangeFactoryBaseNode::BuildFactoryNodeUid(ShaderGraphNodeUid);
-		FString MaterialFactoryNodeName = ShaderGraphNode.GetDisplayLabel();
-
-		UInterchangeMaterialInstanceFactoryNode* NewFactoryNode = NewObject<UInterchangeMaterialInstanceFactoryNode>(&NodeContainer);
-		NewFactoryNode->InitializeNode(MaterialFactoryNodeUid, MaterialFactoryNodeName, EInterchangeNodeContainerType::FactoryData);
-		NewFactoryNode->AddTargetNodeUid(ShaderGraphNodeUid);
+				TArray<FString> TargetNodeUids;
+				FactoryNode->GetTargetNodeUids(TargetNodeUids);
 
-		const bool bAddSourceNodeName = false;
-		UInterchangeUserDefinedAttributesAPI::DuplicateAllUserDefinedAttribute(&ShaderGraphNode, NewFactoryNode, bAddSourceNodeName);
+				for (const FString& TargetNodeUid : TargetNodeUids)
+				{
+					const UInterchangeMaterialInstanceNode* TranslatedNode = Cast<UInterchangeMaterialInstanceNode>(
+						NodeContainer.GetNode(TargetNodeUid)
+					);
 
-		SetCustomParentMaterial(ShaderGraphNode, *NewFactoryNode);
+					if (!TranslatedNode)
+					{
+						continue;
+					}
 
-		// Pick material instance class
-		{
-			UInterchangeEditorUtilitiesBase* EditorUtilities = UInterchangeManager::GetInterchangeManager().GetEditorUtilities();
-			const UClass* MaterialClass = (EditorUtilities && EditorUtilities->IsRuntimeOrPIE()) ? UMaterialInstanceDynamic::StaticClass()
-																								 : UMaterialInstanceConstant::StaticClass();
-			NewFactoryNode->SetCustomInstanceClassName(MaterialClass->GetPathName());
-		}
+					// We explicitly flag the material nodes to parse with this attribute
+					bool bParse;
+					if (!TranslatedNode->GetBooleanAttribute(*UE::Interchange::USD::ParseMaterialIdentifier, bParse) || !bParse)
+					{
+						continue;
+					}
 
-		// Move material attributes to the factory nodes
-		TArray<UE::Interchange::FAttributeKey> AllAttributeKeys;
-		ShaderGraphNode.GetAttributeKeys(AllAttributeKeys);
-		for (const UE::Interchange::FAttributeKey& AttributeKey : AllAttributeKeys)
-		{
-			const FString AttributeKeyString = AttributeKey.ToString();
+					TArray<UE::Interchange::FAttributeKey> AllAttributeKeys;
+					TranslatedNode->GetAttributeKeys(AllAttributeKeys);
 
-			// It's a regular material attribute
-			if (AttributeKeyString.Contains(UE::Interchange::USD::AttributeIdentifierPrefix))
-			{
-				const FString CleanAttributeKey = AttributeKeyString.Replace(	 //
-					*UE::Interchange::USD::AttributeIdentifierPrefix,
-					TEXT(""),
-					ESearchCase::CaseSensitive
-				);
-				const FString StandardParameterName = UInterchangeShaderPortsAPI::MakeInputValueKey(CleanAttributeKey);
+					for (const UE::Interchange::FAttributeKey& AttributeKey : AllAttributeKeys)
+					{
+						const FString AttributeKeyString = AttributeKey.ToString();
 
-				const UE::Interchange::EAttributeTypes AttributeType = ShaderGraphNode.GetAttributeType(AttributeKey);
-				switch (AttributeType)
-				{
-					case UE::Interchange::EAttributeTypes::Bool:
-						if (bool Value; ShaderGraphNode.GetBooleanAttribute(AttributeKey.Key, Value))
+						// It's another attribute related to primvar-compatible materials, just move it over to the factory node as-is
+						if (AttributeKeyString.StartsWith(UE::Interchange::USD::ParameterToPrimvarAttributePrefix))
 						{
-							NewFactoryNode->AddBooleanAttribute(StandardParameterName, Value);
-						}
-						break;
-					case UE::Interchange::EAttributeTypes::Float:
-						if (float Value; ShaderGraphNode.GetFloatAttribute(AttributeKey.Key, Value))
-						{
-							NewFactoryNode->AddFloatAttribute(StandardParameterName, Value);
-						}
-						break;
-					case UE::Interchange::EAttributeTypes::LinearColor:
-						if (FLinearColor Value; ShaderGraphNode.GetLinearColorAttribute(AttributeKey.Key, Value))
-						{
-							NewFactoryNode->AddLinearColorAttribute(StandardParameterName, Value);
+							FString AttributeValue;
+							if (TranslatedNode->GetStringAttribute(AttributeKey.Key, AttributeValue))
+							{
+								FactoryNode->AddStringAttribute(AttributeKeyString, AttributeValue);
+							}
 						}
-						break;
-					case UE::Interchange::EAttributeTypes::String:
-						if (FString TextureTranslatedNodeUid; ShaderGraphNode.GetStringAttribute(AttributeKey.Key, TextureTranslatedNodeUid))
+						else if (AttributeKeyString.StartsWith(UE::Interchange::USD::PrimvarUVIndexAttributePrefix))
 						{
-							// We only use string attributes for textures for now
-							FString FactoryTextureUid = UInterchangeFactoryBaseNode::BuildFactoryNodeUid(TextureTranslatedNodeUid);
-
-							NewFactoryNode->AddStringAttribute(StandardParameterName, FactoryTextureUid);
-							NewFactoryNode->AddFactoryDependencyUid(FactoryTextureUid);
+							int32 AttributeValue;
+							if (TranslatedNode->GetInt32Attribute(AttributeKey.Key, AttributeValue))
+							{
+								FactoryNode->AddInt32Attribute(AttributeKeyString, AttributeValue);
+							}
 						}
-						break;
-					default:
-						break;
-				}
-			}
-			// It's another attribute related to primvar-compatible materials, just move it over to the factory node as-is
-			else if (AttributeKeyString.StartsWith(UE::Interchange::USD::ParameterToPrimvarAttributePrefix))
-			{
-				FString AttributeValue;
-				if (ShaderGraphNode.GetStringAttribute(AttributeKey.Key, AttributeValue))
-				{
-					NewFactoryNode->AddStringAttribute(AttributeKeyString, AttributeValue);
-				}
-			}
-			else if (AttributeKeyString.StartsWith(UE::Interchange::USD::PrimvarUVIndexAttributePrefix))
-			{
-				int32 AttributeValue;
-				if (ShaderGraphNode.GetInt32Attribute(AttributeKey.Key, AttributeValue))
-				{
-					NewFactoryNode->AddInt32Attribute(AttributeKeyString, AttributeValue);
-				}
-			}
-		}
-
-		return NewFactoryNode;
-	}
-
-	// Replaces dummy material instance factory nodes with new nodes where we properly fill in the
-	// material parameters and parent materials.
-	//
-	// We always translate the USD materials into UInterchangeShaderGraphNode (as opposed to directly
-	// UInterchangeMaterialInstanceNode) and fill in all parameters as base translated node attributes,
-	// in order to defer to the UInterchangeGenericMaterialPipeline (or other pipelines) whether to emit
-	// new material graphs or material instances.
-	//
-	// Here, we will likely run *after* UInterchangeGenericMaterialPipeline, which means is will already have
-	// converted our UInterchangeShaderGraphNodes into either material or material instance factory nodes,
-	// according to its "MaterialImport" option.
-	//
-	// If MaterialImport is set to "ImportAsMaterialInstances" however, UInterchangeGenericMaterialPipeline
-	// will not convert our base translated node attributes into the factory node like we need it to.
-	// This is what this function aims to repair: It will go over any of these broken material instance factory
-	// nodes, and replace them with new nodes that have the correct material info extracted from the
-	// corresponding ("target") UInterchangeShaderGraphNode base node attributes. It is also in charge of
-	// picking the correct parent material for these nodes, according to all the info added to the shader graph
-	// nodes.
-	//
-	// Note: This is all for the material instance code path. Unfortunately the full material graph code path
-	// won't work either, as we don't emit the shader graph translated node data from the USD translator just yet.
-	// We will likely want to support arbitrary shader graph generation for that though, which may come at a
-	// later point
-	void RepairMaterialInstanceFactoryNodes(UInterchangeBaseNodeContainer& NodeContainer)
-	{
-		TMap<FString, const UInterchangeShaderGraphNode*> MaterialFactoryNodeUidsToShaderGraphNodes;
-
-		auto CollectNodes =
-			[&NodeContainer, &MaterialFactoryNodeUidsToShaderGraphNodes](const FString& NodeUid, UInterchangeMaterialInstanceFactoryNode* FactoryNode)
-		{
-			TArray<FString> TargetNodeUids;
-			FactoryNode->GetTargetNodeUids(TargetNodeUids);
-
-			for (const FString& TargetNodeUid : TargetNodeUids)
-			{
-				const UInterchangeShaderGraphNode* ShaderGraphNode = Cast<UInterchangeShaderGraphNode>(NodeContainer.GetNode(TargetNodeUid));
-				if (!ShaderGraphNode)
-				{
-					continue;
-				}
-
-				// We explicitly flag our shader graph nodes with this attribute
-				bool bParse;
-				if (ShaderGraphNode->GetBooleanAttribute(*UE::Interchange::USD::ParseMaterialIdentifier, bParse) && bParse)
-				{
-					MaterialFactoryNodeUidsToShaderGraphNodes.Add(NodeUid, ShaderGraphNode);
-					break;
+					}
 				}
 			}
-		};
-		NodeContainer.IterateNodesOfType<UInterchangeMaterialInstanceFactoryNode>(CollectNodes);
-
-		for (const TPair<FString, const UInterchangeShaderGraphNode*>& Pair : MaterialFactoryNodeUidsToShaderGraphNodes)
-		{
-			const FString& FactoryNodeUid = Pair.Key;
-			const UInterchangeShaderGraphNode* ShaderGraphNode = Pair.Value;
-
-			UInterchangeMaterialInstanceFactoryNode* NewFactoryNode = BuildMaterialInstance(NodeContainer, *ShaderGraphNode);
-			if (NewFactoryNode)
-			{
-				NodeContainer.ReplaceNode(FactoryNodeUid, NewFactoryNode);
-			}
-		}
+		);
 	}
 
 	// Move custom attribute info from first volume nodes into corresponding factory nodes.
@@ -1474,7 +1299,7 @@ void UInterchangeUsdPipeline::ExecutePipeline(
 
 	DisableNodesBasedOnGeometryPurpose(*NodeContainer, (EUsdPurpose)GeometryPurpose);
 
-	RepairMaterialInstanceFactoryNodes(*NodeContainer);
+	MoveCustomAttributesToMaterialFactoryNodes(*NodeContainer);
 
 	ProcessVolumeNodes(*NodeContainer);
 
diff --git a/Engine/Plugins/Interchange/Extensions/OpenUSD/Source/Import/Private/InterchangeUsdTranslator.cpp b/Engine/Plugins/Interchange/Extensions/OpenUSD/Source/Import/Private/InterchangeUsdTranslator.cpp
--- a/Engine/Plugins/Interchange/Extensions/OpenUSD/Source/Import/Private/InterchangeUsdTranslator.cpp
+++ b/Engine/Plugins/Interchange/Extensions/OpenUSD/Source/Import/Private/InterchangeUsdTranslator.cpp
@@ -272,14 +272,16 @@ namespace UE::InterchangeUsdTranslator::Private
 	public:
 		/** Add a material instance to the node container, otherwise it will add a material if it comes from a Translator (for example coming from
 		 * MaterialX which cannot handle material instances) */
-		TArray<FString> AddMaterialNode(
+		FString AddMaterialNode(
 			const UE::FUsdPrim& Prim,
 			UInterchangeUsdTranslatorSettings* TranslatorSettings,
-			UInterchangeBaseNodeContainer& NodeContainer
+			UInterchangeBaseNodeContainer& NodeContainer,
+			bool bForceTwoSided = false
 		);
 
 		FString AddMeshNode(
 			const UE::FUsdPrim& Prim,
+			UInterchangeUsdTranslatorSettings* TranslatorSettings,
 			UInterchangeBaseNodeContainer& NodeContainer,
 			const FTraversalInfo& Info,
 			bool bPrimitiveShape = false
@@ -312,18 +314,6 @@ namespace UE::InterchangeUsdTranslator::Private
 			TArray<FUsdMaterialXShaderGraph::FGeomProp>& OutGeomPropValues
 		);
 
-		/**
-		 * Set the material slot dependencies on the meshes, if bIterateShaderGraphs is true, it will iterate over the shader graphs to rename them
-		 * according to the MaterialPrimName
-		 */
-		void SetMaterialSlotDependencies(
-			const FString& MaterialPrimName,
-			const FString& MaterialUid,
-			UInterchangeBaseNodeContainer& NodeContainer,
-			const TArray<FUsdMaterialXShaderGraph::FGeomProp>& GeomPropValues = {},
-			bool bIterateShaderGraphs = false
-		);
-
 		/**
 		 * If we're not translating a decompressed USD root, returns TexturePathOnDisk.
 		 * If we are translating a decompressed USD root, returns the path to the USDZ file itself.
@@ -393,13 +383,7 @@ namespace UE::InterchangeUsdTranslator::Private
 		TMap<FString, TArray<UsdUtils::FVolumePrimInfo>> PrimPathToVolumeInfo;
 
 	private:
-		struct FMaterialSlotMesh
-		{
-			FString MaterialSlotName;
-			UInterchangeMeshNode* MeshNode;
-		};
-		TMap<FString, FString> MaterialUidToActualNodeUid;
-		TMap<FString, TArray<FMaterialSlotMesh>> PrimPathToSlotMeshNodes;
+		TMap<FString, TArray<FUsdMaterialXShaderGraph::FGeomProp>> MaterialUidToGeomProps;
 
 		// Used within a translation. We cache these because we make a volume node *per .vdb file*, and on the USD side we may have
 		// any number of Volume prims internally using the same .vdb file, and we want to share these whenever possible
@@ -870,14 +854,8 @@ namespace UE::InterchangeUsdTranslator::Private
 		FString DisplayColorDesc = Slot.MaterialSource;
 		FString NodeUid = MaterialPrefix + Slot.MaterialSource;
 
-		if (Slot.bMeshIsDoubleSided)
-		{
-			NodeUid += TwoSidedSuffix;
-			DisplayColorDesc += TwoSidedSuffix;
-		}
-
 		// We'll treat the DisplayColorDesc (something like "!DisplayColor_1_0") as the material instance UID here
-		const UInterchangeShaderGraphNode* Node = Cast<UInterchangeShaderGraphNode>(NodeContainer.GetNode(NodeUid));
+		const UInterchangeMaterialInstanceNode* Node = Cast<UInterchangeMaterialInstanceNode>(NodeContainer.GetNode(NodeUid));
 		if (Node)
 		{
 			return NodeUid;
@@ -891,16 +869,14 @@ namespace UE::InterchangeUsdTranslator::Private
 		}
 		FString NodeName = ParsedMat->ToPrettyString();
 
-		UInterchangeShaderGraphNode* NewNode = NewObject<UInterchangeShaderGraphNode>(&NodeContainer);
+		UInterchangeMaterialInstanceNode* NewNode = NewObject<UInterchangeMaterialInstanceNode>(&NodeContainer);
 		NewNode->InitializeNode(NodeUid, NodeName, EInterchangeNodeContainerType::TranslatedAsset);
-		NewNode->AddBooleanAttribute(UE::Interchange::USD::ParseMaterialIdentifier, true);
-		NewNode->AddStringAttribute(*UE::Interchange::USD::DisplayColorDescIdentifier, DisplayColorDesc);
 
-		if (Slot.bMeshIsDoubleSided)
+		if (const FSoftObjectPath* ParentMaterialPath = GetReferenceMaterialPath(ParsedMat.GetValue()))
 		{
-			NewNode->SetCustomTwoSided(true);
-			NewNode->SetCustomTwoSidedTransmission(true); // weirdly enough we also have to set this in case of a translucent material
+			NewNode->SetCustomParent(ParentMaterialPath->GetAssetPathString());
 		}
+
 		NodeContainer.AddNode(NewNode);
 
 		return NodeUid;
@@ -930,99 +906,147 @@ namespace UE::InterchangeUsdTranslator::Private
 		return NodeUid;
 	}
 
-	void UInterchangeUSDTranslatorImpl::SetMaterialSlotDependencies(
-		const FString& MaterialPrimName,
-		const FString& MaterialUid,
-		UInterchangeBaseNodeContainer& NodeContainer,
-		const TArray<FUsdMaterialXShaderGraph::FGeomProp>& GeomPropValues,
-		bool bIterateShaderGraphs
-	)
+	// Returns the UID of the material translated node that was generated from the MaterialX translation of a particular
+	// material prim.
+	//
+	// This works because when parsing MaterialX files we generate shader graph nodes with UIDs that match the original
+	// material prim name in the USD file (e.g. on USD we have a binding relationship to </MaterialX/Materials/Marble_3D>,
+	// and we end up generating a shader graph node with Uid \\Shaders\\Marble_3D).
+	FString GetMaterialXMaterialUid(const FString& PrimName, UInterchangeBaseNodeContainer& NodeContainer)
 	{
-		using namespace UE::Interchange;
-		if (bIterateShaderGraphs)
-		{
-			// The material from the MaterialXTranslator doesn't have the same UID, both the PrimPath have the same name but not the same path
-			// We need to retrieve that name (which is the Material name) in the translator, then we can map it to the right mesh
-			NodeContainer.BreakableIterateNodesOfType<UInterchangeShaderGraphNode>(
-				[this, &MaterialPrimName, &MaterialUid](const FString&, UInterchangeShaderGraphNode* ShaderGraphNode)
+		FString Result;
+
+		NodeContainer.BreakableIterateNodesOfType<UInterchangeShaderGraphNode>(
+			[&NodeContainer, &PrimName, &Result](const FString&, UInterchangeShaderGraphNode* ShaderGraphNode)
+			{
+				FString NodeUid = ShaderGraphNode->GetUniqueID();
+				if (FPaths::GetBaseFilename(NodeUid) == PrimName)
 				{
-					FString UID = ShaderGraphNode->GetUniqueID();
-					if (FPaths::GetBaseFilename(UID) == MaterialPrimName)
-					{
-						MaterialUidToActualNodeUid.Add(MaterialUid, UID);
-						return true;
-					}
-					else
-					{
-						return false;
-					}
+					Result = NodeUid;
+					return true;
+				}
+				else
+				{
+					return false;
 				}
+			}
+		);
+
+		return Result;
+	}
+
+	FString GetOrCreateTwoSidedShaderGraphNode(const FString& OneSidedShaderGraphNodeUid, UInterchangeBaseNodeContainer& NodeContainer)
+	{
+		FString TwoSidedUid = OneSidedShaderGraphNodeUid + TwoSidedSuffix;
+
+		// We already created this, just return it
+		if (const UInterchangeShaderGraphNode* Node = Cast<UInterchangeShaderGraphNode>(NodeContainer.GetNode(TwoSidedUid)))
+		{
+			return TwoSidedUid;
+		}
+
+		const UInterchangeShaderGraphNode* OneSidedNode = Cast<UInterchangeShaderGraphNode>(NodeContainer.GetNode(OneSidedShaderGraphNodeUid));
+		if (!OneSidedNode)
+		{
+			UE_LOG(
+				LogUsd,
+				Warning,
+				TEXT("Failed to create TwoSided material, as the UID '%s' does not point to a one-sided material!"),
+				*OneSidedShaderGraphNodeUid
 			);
+			return OneSidedShaderGraphNodeUid;
+		}
+
+		// One-sided node is already two-sided, so just return that.
+		// This can happen e.g. if the MaterialX translator internally sets the node as two-sided because the shader graph says to do that.
+		// Note that we don't have anything caring for the exact opposite: If the USD Mesh is explicitly one-sided and the MaterialX material is
+		// two-sided in this way we'll just use the two-sided material on the mesh. For now let's presume that's what the user intended, as you
+		// have to explicitly set the MaterialX material as two-sided for that
+		bool bIsTwoSided = false;
+		if (OneSidedNode->GetCustomTwoSided(bIsTwoSided) && bIsTwoSided)
+		{
+			return OneSidedShaderGraphNodeUid;
 		}
 
-		// Let's iterate over the Shader Nodes, the TextureSample ones more specifically,
-		// to see if have to retrieve any attributes related to the conversion of geompropvalues
+		FString TwoSidedNodeName = OneSidedNode->GetDisplayLabel() + TwoSidedSuffix;
+		UInterchangeShaderGraphNode* TwoSidedNode = NewObject<UInterchangeShaderGraphNode>(&NodeContainer);
+
+		UInterchangeBaseNode::CopyStorage(OneSidedNode, TwoSidedNode);
+		TwoSidedNode->InitializeNode(TwoSidedUid, TwoSidedNodeName, EInterchangeNodeContainerType::TranslatedAsset);
+		TwoSidedNode->SetCustomTwoSided(true);
+		TwoSidedNode->SetCustomTwoSidedTransmission(true);
+
+		NodeContainer.AddNode(TwoSidedNode);
+		return TwoSidedUid;
+	}
+
+	// Sets on the provided MeshNode custom attributes needed to bake the provided GeomPropValue geomprops / primvars
+	// into textures, for the provided MaterialNodeUid
+	void AddPrimvarBakingAttributes(
+		UInterchangeMeshNode* MeshNode,
+		const FString& MaterialNodeUid,
+		UInterchangeBaseNodeContainer& NodeContainer,
+		const TArray<FUsdMaterialXShaderGraph::FGeomProp>& GeomPropValues
+	)
+	{
 #if WITH_EDITOR
+		using namespace UE::Interchange;
+
+		if (!MeshNode)
+		{
+			return;
+		}
+
 		// Let's iterate over the Shader Nodes, the TextureSample ones, more specifically,
 		// to see if we have to retrieve any attributes related to the conversion of geompropvalues
 		// we'll store the UID of the shader node in order to retrieve it during the baking phase (in the post factory import)
 		TArray<FString> ShaderNodesTextureSampleUIDs;
 		NodeContainer.IterateNodesOfType<UInterchangeShaderNode>(
-			[&ShaderNodesTextureSampleUIDs](const FString&, UInterchangeShaderNode* ShaderNode)
+			[&ShaderNodesTextureSampleUIDs, &MaterialNodeUid](const FString& ShaderUid, UInterchangeShaderNode* ShaderNode)
 			{
 				using namespace UE::Interchange::Materials::Standard::Nodes;
 
-				constexpr const TCHAR* GeomPropHandledAttribute = TEXT("GeomPropHandled");
+				// We only care about baking the geomprop nodes that were generated when parsing this Material, and
+				// they should always have the Uid of the material as a prefix
+				if (!ShaderUid.StartsWith(MaterialNodeUid))
+				{
+					return;
+				}
+
 				if (FString ShaderType; ShaderNode->GetCustomShaderType(ShaderType) && ShaderType == TextureSample::Name.ToString())
 				{
 					if (bool bIsGeomProp; ShaderNode->GetBooleanAttribute(::MaterialX::Attributes::GeomPropImage, bIsGeomProp) && bIsGeomProp)
 					{
-						// We only want the geomprops that were not already assigned to the mesh node
-						if (bool bGeomPropHandled; !ShaderNode->GetBooleanAttribute(GeomPropHandledAttribute, bGeomPropHandled) || !bGeomPropHandled)
-						{
-							ShaderNode->AddBooleanAttribute(GeomPropHandledAttribute, true);
-							ShaderNodesTextureSampleUIDs.Emplace(ShaderNode->GetUniqueID());
-						}
+						ShaderNodesTextureSampleUIDs.Emplace(ShaderNode->GetUniqueID());
 					}
 				}
 			}
 		);
-#endif	  // WITH_EDITOR
 
-		// Now we need to check if we have to set the slot of the mesh nodes here
-		if (TArray<FMaterialSlotMesh>* SlotMeshes = PrimPathToSlotMeshNodes.Find(MaterialUid))
+		if (GeomPropValues.Num() != ShaderNodesTextureSampleUIDs.Num())
 		{
-			if (FString* NewMaterialUID = MaterialUidToActualNodeUid.Find(MaterialUid))
-			{
-				for (const FMaterialSlotMesh& MaterialSlotMesh : *SlotMeshes)
-				{
-#if WITH_EDITOR
-					if (!GeomPropValues.IsEmpty())
-					{
-						MaterialSlotMesh.MeshNode->AddInt32Attribute(USD::Primvar::Number, GeomPropValues.Num());
+			UE_LOG(
+				LogUsd,
+				Warning,
+				TEXT("Failed to bake primvars for mesh '%s' and material '%s': Encountered %d GeomPropValues but %d geomprop shader nodes!"),
+				*MeshNode->GetUniqueID(),
+				*MaterialNodeUid,
+				GeomPropValues.Num(),
+				ShaderNodesTextureSampleUIDs.Num()
+			);
+			return;
+		}
 
-						for (int32 Index = 0; Index < GeomPropValues.Num(); ++Index)
-						{
-							MaterialSlotMesh.MeshNode->AddStringAttribute(USD::Primvar::Name + FString::FromInt(Index), GeomPropValues[Index].Name);
-							MaterialSlotMesh.MeshNode->AddBooleanAttribute(
-								USD::Primvar::TangentSpace + FString::FromInt(Index),
-								GeomPropValues[Index].bTangentSpace
-							);
-							MaterialSlotMesh.MeshNode->AddStringAttribute(
-								USD::Primvar::ShaderNodeTextureSample + FString::FromInt(Index),
-								ShaderNodesTextureSampleUIDs[Index]
-							);
-						}
-					}
-#endif	  // WITH_EDITOR
-					if (!MaterialSlotMesh.MeshNode->GetSlotMaterialDependencyUid(MaterialSlotMesh.MaterialSlotName, *NewMaterialUID))
-					{
-						MaterialSlotMesh.MeshNode->SetSlotMaterialDependencyUid(MaterialSlotMesh.MaterialSlotName, *NewMaterialUID);
-					}
-				}
-			}
+		MeshNode->AddInt32Attribute(USD::Primvar::Number, GeomPropValues.Num());
+
+		for (int32 Index = 0; Index < GeomPropValues.Num(); ++Index)
+		{
+			MeshNode->AddStringAttribute(USD::Primvar::Name + FString::FromInt(Index), GeomPropValues[Index].Name);
+			MeshNode->AddBooleanAttribute(USD::Primvar::TangentSpace + FString::FromInt(Index), GeomPropValues[Index].bTangentSpace);
+			MeshNode->AddStringAttribute(USD::Primvar::ShaderNodeTextureSample + FString::FromInt(Index), ShaderNodesTextureSampleUIDs[Index]);
 		}
-	};
+#endif	  // WITH_EDITOR
+	}
 
 	bool UInterchangeUSDTranslatorImpl::AddMaterialXShaderGraph(
 		const UE::FUsdPrim& Prim,
@@ -1096,7 +1120,7 @@ namespace UE::InterchangeUsdTranslator::Private
 		return false;
 	}
 
-	// We use this visitor to set UsdToUnreal::FParameterValue TVariant values onto UInterchangeShaderGraphNode.
+	// We use this visitor to set UsdToUnreal::FParameterValue TVariant values onto UInterchangeMaterialInstanceNode.
 	//
 	// For now we only set attributes meant to be parsed as material instance parameters.
 	// If/whenever we want to support generating full material shader graphs from USD, we likely don't want to just fill out inputs
@@ -1108,7 +1132,7 @@ namespace UE::InterchangeUsdTranslator::Private
 		FMaterialInstanceParameterValueVisitor(
 			const UE::FUsdPrim& InPrim,
 			UInterchangeBaseNodeContainer& InNodeContainer,
-			UInterchangeShaderGraphNode& InMaterialNode,
+			UInterchangeMaterialInstanceNode& InMaterialNode,
 			const TMap<FString, int32>& InPrimvarToUVIndex
 		)
 			: Prim(InPrim)
@@ -1118,30 +1142,22 @@ namespace UE::InterchangeUsdTranslator::Private
 		{
 		}
 
-		static FString MakeInstanceParameterName(const FString& InParameterName)
-		{
-			return UE::Interchange::USD::AttributeIdentifierPrefix + InParameterName;
-		}
-
 		void EnableTextureForChannel(bool bEnable) const
 		{
 			using namespace UE::Interchange::USD;
 
-			MaterialNode.AddFloatAttribute(
-				MakeInstanceParameterName(UseTextureParameterPrefix + *BaseParameterName + UseTextureParameterSuffix),
-				bEnable ? 1.0f : 0.0f
-			);
+			MaterialNode.AddScalarParameterValue(UseTextureParameterPrefix + *BaseParameterName + UseTextureParameterSuffix, bEnable ? 1.0f : 0.0f);
 		}
 
 		void operator()(const float Value) const
 		{
-			MaterialNode.AddFloatAttribute(MakeInstanceParameterName(*BaseParameterName), Value);
+			MaterialNode.AddScalarParameterValue(*BaseParameterName, Value);
 			EnableTextureForChannel(false);
 		}
 
 		void operator()(const FVector& Value) const
 		{
-			MaterialNode.AddLinearColorAttribute(MakeInstanceParameterName(*BaseParameterName), FLinearColor{Value});
+			MaterialNode.AddVectorParameterValue(*BaseParameterName, FLinearColor{Value});
 			EnableTextureForChannel(false);
 		}
 
@@ -1154,8 +1170,8 @@ namespace UE::InterchangeUsdTranslator::Private
 			AddTextureNode(Prim, TextureUid, Value, NodeContainer);
 
 			// Actual texture assignment
-			MaterialNode.AddStringAttribute(	//
-				MakeInstanceParameterName(FString::Printf(TEXT("%sTexture"), **BaseParameterName)),
+			MaterialNode.AddTextureParameterValue(	  //
+				FString::Printf(TEXT("%sTexture"), **BaseParameterName),
 				TextureUid
 			);
 			EnableTextureForChannel(true);
@@ -1167,20 +1183,17 @@ namespace UE::InterchangeUsdTranslator::Private
 				Value.UVTranslation[0],
 				Value.UVTranslation[1]
 			};
-			MaterialNode.AddLinearColorAttribute(
-				MakeInstanceParameterName(FString::Printf(TEXT("%sScaleTranslation"), **BaseParameterName)),
-				ScaleAndTranslation
-			);
-			MaterialNode.AddFloatAttribute(	   //
-				MakeInstanceParameterName(FString::Printf(TEXT("%sRotation"), **BaseParameterName)),
+			MaterialNode.AddVectorParameterValue(FString::Printf(TEXT("%sScaleTranslation"), **BaseParameterName), ScaleAndTranslation);
+			MaterialNode.AddScalarParameterValue(	 //
+				FString::Printf(TEXT("%sRotation"), **BaseParameterName),
 				Value.UVRotation
 			);
 
 			// UV index
 			if (const int32* FoundIndex = PrimvarToUVIndex.Find(Value.Primvar))
 			{
-				MaterialNode.AddFloatAttribute(	   //
-					MakeInstanceParameterName(*BaseParameterName + UE::Interchange::USD::UVIndexParameterSuffix),
+				MaterialNode.AddScalarParameterValue(	 //
+					*BaseParameterName + UE::Interchange::USD::UVIndexParameterSuffix,
 					*FoundIndex
 				);
 			}
@@ -1223,39 +1236,38 @@ namespace UE::InterchangeUsdTranslator::Private
 					ComponentMask = FLinearColor{0.f, 0.f, 0.f, 1.f};
 					break;
 			}
-			MaterialNode.AddLinearColorAttribute(
-				MakeInstanceParameterName(FString::Printf(TEXT("%sTextureComponent"), **BaseParameterName)),
-				ComponentMask
-			);
+			MaterialNode.AddVectorParameterValue(FString::Printf(TEXT("%sTextureComponent"), **BaseParameterName), ComponentMask);
 		}
 
 		void operator()(const UsdToUnreal::FPrimvarReaderParameterValue& Value) const
 		{
-			MaterialNode.AddLinearColorAttribute(MakeInstanceParameterName(*BaseParameterName), FLinearColor{Value.FallbackValue});
+			MaterialNode.AddVectorParameterValue(*BaseParameterName, FLinearColor{Value.FallbackValue});
 
 			if (Value.PrimvarName == TEXT("displayColor"))
 			{
-				MaterialNode.AddFloatAttribute(MakeInstanceParameterName(TEXT("UseVertexColorForBaseColor")), 1.0f);
+				MaterialNode.AddScalarParameterValue(TEXT("UseVertexColorForBaseColor"), 1.0f);
 			}
 		}
 
 		void operator()(const bool Value) const
 		{
-			MaterialNode.AddBooleanAttribute(MakeInstanceParameterName(*BaseParameterName), static_cast<float>(Value));
+			// Actual booleans are only meant for static switches on Interchange
+			MaterialNode.AddScalarParameterValue(*BaseParameterName, static_cast<float>(Value));
 		}
 
 	public:
 		const UE::FUsdPrim& Prim;
 		UInterchangeBaseNodeContainer& NodeContainer;
-		UInterchangeShaderGraphNode& MaterialNode;
+		UInterchangeMaterialInstanceNode& MaterialNode;
 		const TMap<FString, int32>& PrimvarToUVIndex;
 		const FString* BaseParameterName = nullptr;
 	};
 
-	TArray<FString> UInterchangeUSDTranslatorImpl::AddMaterialNode(
+	FString UInterchangeUSDTranslatorImpl::AddMaterialNode(
 		const UE::FUsdPrim& Prim,
 		UInterchangeUsdTranslatorSettings* TranslatorSettings,
-		UInterchangeBaseNodeContainer& NodeContainer
+		UInterchangeBaseNodeContainer& NodeContainer,
+		bool bForceTwoSided
 	)
 	{
 		TRACE_CPUPROFILER_EVENT_SCOPE(UInterchangeUSDTranslatorImpl::AddMaterialNode)
@@ -1273,54 +1285,49 @@ namespace UE::InterchangeUsdTranslator::Private
 			TOptional<FString> UnrealContentPath = UsdUtils::GetUnrealSurfaceOutput(Prim);
 			if (UnrealContentPath.IsSet())
 			{
-				return { UE::InterchangeUsdTranslator::Private::AddUnrealMaterialReferenceNodeIfNeeded(NodeContainer, UnrealContentPath.GetValue()) };
+				return {UE::InterchangeUsdTranslator::Private::AddUnrealMaterialReferenceNodeIfNeeded(NodeContainer, UnrealContentPath.GetValue())};
 			}
 		}
 
-		TArray<FString> MaterialUids;
-
 		FString PrimPath = Prim.GetPrimPath().GetString();
+		FString MaterialPrimName = Prim.GetName().ToString();
 		FString MaterialUid = MaterialPrefix + PrimPath;
-		FString MaterialTwoSidedUid = MaterialUid + TwoSidedSuffix;
-		FString MaterialPrimName(Prim.GetName().ToString());
-		bool bDuplicateMaterialToBeTwoSided = PrimPathToSlotMeshNodes.Contains(MaterialTwoSidedUid);
+		FString MaterialNodeName = MaterialPrimName;
 
-		MaterialUids.Emplace(MaterialUid);
-		if (bDuplicateMaterialToBeTwoSided)
+		if (bForceTwoSided)
 		{
-			MaterialUids.Emplace(MaterialTwoSidedUid);
+			MaterialUid += TwoSidedSuffix;
+			MaterialNodeName += TwoSidedSuffix;
 		}
 
-		// Check if Node already exist as a TwoSided material
-		if (const UInterchangeShaderGraphNode* Node = Cast<UInterchangeShaderGraphNode>(NodeContainer.GetNode(MaterialTwoSidedUid)))
+		if (const UInterchangeMaterialInstanceNode* Node = Cast<UInterchangeMaterialInstanceNode>(NodeContainer.GetNode(MaterialUid)))
 		{
-			return MaterialUids;
-		}
-
-		// Check if Node already exist with this ID
-		if (const UInterchangeShaderGraphNode* ShaderGraphNode = Cast<UInterchangeShaderGraphNode>(NodeContainer.GetNode(MaterialUid)))
-		{
-			return MaterialUids;
+			return MaterialUid;
 		}
 
 		// we only create material instances if we didn't find any MaterialX instances (including a shader graph)
 		if (TArray<FUsdMaterialXShaderGraph::FGeomProp> GeomPropValues;
 			AddMaterialXShaderGraph(Prim, TranslatorSettings, NodeContainer, GeomPropValues))
 		{
-			SetMaterialSlotDependencies(MaterialPrimName, MaterialUid, NodeContainer, GeomPropValues, true);
-			return MaterialUids;
+			FString MaterialXMaterialUid = GetMaterialXMaterialUid(MaterialPrimName, NodeContainer);
+			if (bForceTwoSided)
+			{
+				MaterialXMaterialUid = GetOrCreateTwoSidedShaderGraphNode(MaterialXMaterialUid, NodeContainer);
+			}
+
+			if (GeomPropValues.Num() > 0)
+			{
+				MaterialUidToGeomProps.Add(MaterialXMaterialUid, GeomPropValues);
+			}
+
+			return MaterialXMaterialUid;
 		}
 
-		UInterchangeShaderGraphNode* MaterialNode = NewObject<UInterchangeShaderGraphNode>(&NodeContainer);
-		MaterialNode->InitializeNode(MaterialUid, MaterialPrimName, EInterchangeNodeContainerType::TranslatedAsset);
-		MaterialNode->SetAssetName(MaterialPrimName);
-		MaterialNode->AddBooleanAttribute(UE::Interchange::USD::ParseMaterialIdentifier, true);
+		UInterchangeMaterialInstanceNode* MaterialNode = NewObject<UInterchangeMaterialInstanceNode>(&NodeContainer);
+		MaterialNode->InitializeNode(MaterialUid, MaterialNodeName, EInterchangeNodeContainerType::TranslatedAsset);
+		MaterialNode->SetAssetName(MaterialNodeName);
 		NodeContainer.AddNode(MaterialNode);
 
-		// Set the material instance node to the correct mesh nodes
-		MaterialUidToActualNodeUid.Add(MaterialUid, MaterialUid);
-		SetMaterialSlotDependencies(MaterialPrimName, MaterialUid, NodeContainer);
-
 		UsdToUnreal::FUsdPreviewSurfaceMaterialData MaterialData;
 		const bool bSuccess = UsdToUnreal::ConvertMaterial(Prim, MaterialData, *RenderContext.ToString());
 
@@ -1368,40 +1375,41 @@ namespace UE::InterchangeUsdTranslator::Private
 
 				MaterialNode->AddInt32Attribute(UE::Interchange::USD::PrimvarUVIndexAttributePrefix + Primvar, UVIndex);
 			}
+
+			// Let the pipeline know that it should process this node and handle these attributes we just added
+			if (MaterialData.PrimvarToUVIndex.Num() > 0 || MaterialData.Parameters.Num() > 0)
+			{
+				MaterialNode->AddBooleanAttribute(UE::Interchange::USD::ParseMaterialIdentifier, true);
+			}
 		}
 
-		// Find and set the relevant info so that we can pick the reference material on the pipeline
-		//
-		// TODO: Proper VT texture support (we'd need to know the texture resolution at this point, and we haven't parsed them yet...).
-		// The way it currently works on Interchange is that the factory will create a VT or nonVT version of the texture to match the
-		// material parameter slot. Since we'll currently never set the VT reference material, it essentially means it will always
-		// downgrade our VT textures to non-VT.
-		// The only exception is how we upgrade the reference material to VT in case we have any UDIM textures a few lines above,
-		// as those are trivial to check for (we don't have to actually load the textures to do it)
+		EUsdReferenceMaterialProperties Properties = EUsdReferenceMaterialProperties::None;
 		if (UsdUtils::IsMaterialTranslucent(MaterialData))
 		{
-			MaterialNode->SetCustomBlendMode(EBlendMode::BLEND_Translucent);
+			Properties |= EUsdReferenceMaterialProperties::Translucent;
+		}
+		if (bForceTwoSided)
+		{
+			Properties |= EUsdReferenceMaterialProperties::TwoSided;
 		}
 		if (bNeedsVTParent)
 		{
-			MaterialNode->AddBooleanAttribute(UE::Interchange::USD::NeedsVTParentIdentifier, true);
+			// TODO: Proper VT texture support (we'd need to know the texture resolution at this point, and we haven't parsed them yet...).
+			// The way it currently works on Interchange is that the factory will create a VT or nonVT version of the texture to match the
+			// material parameter slot. Since we'll currently never set the VT reference material, it essentially means it will always
+			// downgrade our VT textures to non-VT.
+			// The only exception is how we upgrade the reference material to VT in case we have any UDIM textures a few lines above,
+			// as those are trivial to check for (we don't have to actually load the textures to do it)
+			Properties |= EUsdReferenceMaterialProperties::VT;
 		}
 
-		// Check if we have to duplicate the material to be TwoSided
-		if (bDuplicateMaterialToBeTwoSided)
+		FSoftObjectPath ParentMaterial = UsdUnreal::MaterialUtils::GetReferencePreviewSurfaceMaterial(Properties);
+		if (ParentMaterial.IsValid())
 		{
-			FString MaterialAssetName = MaterialPrimName + TwoSidedSuffix;
-			UInterchangeShaderGraphNode* MaterialTwoSidedNode = NewObject<UInterchangeShaderGraphNode>(&NodeContainer);
-			UInterchangeBaseNode::CopyStorage(MaterialNode, MaterialTwoSidedNode);
-			MaterialTwoSidedNode->InitializeNode(MaterialTwoSidedUid, MaterialAssetName, EInterchangeNodeContainerType::TranslatedAsset);
-			MaterialTwoSidedNode->SetCustomTwoSided(true);
-			MaterialTwoSidedNode->SetCustomTwoSidedTransmission(true);
-			NodeContainer.AddNode(MaterialTwoSidedNode);
-
-			MaterialUidToActualNodeUid.Add(MaterialTwoSidedUid, MaterialTwoSidedUid);
-			SetMaterialSlotDependencies(MaterialAssetName, MaterialTwoSidedUid, NodeContainer);
+			MaterialNode->SetCustomParent(ParentMaterial.GetAssetPathString());
 		}
-		return MaterialUids;
+
+		return MaterialUid;
 	}
 
 	FString UInterchangeUSDTranslatorImpl::GetTextureSourcePath(const FString& TexturePathOnDisk)
@@ -1907,6 +1915,7 @@ namespace UE::InterchangeUsdTranslator::Private
 
 	FString UInterchangeUSDTranslatorImpl::AddMeshNode(
 		const UE::FUsdPrim& Prim,
+		UInterchangeUsdTranslatorSettings* TranslatorSettings,
 		UInterchangeBaseNodeContainer& NodeContainer,
 		const FTraversalInfo& Info,
 		bool bPrimitiveShape
@@ -1918,6 +1927,8 @@ namespace UE::InterchangeUsdTranslator::Private
 		FString NodeUid = (bPrimitiveShape ? PrimitiveShapePrefix : MeshPrefix) + PrimPath;
 		FString NodeName(Prim.GetName().ToString());
 
+		UE::FUsdStage Stage = Prim.GetStage();
+
 		// Check if Node already exist with this ID
 		if (const UInterchangeMeshNode* Node = Cast<UInterchangeMeshNode>(NodeContainer.GetNode(NodeUid)))
 		{
@@ -1956,7 +1967,6 @@ namespace UE::InterchangeUsdTranslator::Private
 		{
 			GeometryCacheNode->SetPayLoadKey(PrimPath, EInterchangeMeshPayLoadType::ANIMATED);
 
-			UE::FUsdStage Stage = Prim.GetStage();
 			int32 StartFrame = FMath::FloorToInt(Stage.GetStartTimeCode());
 			int32 EndFrame = FMath::CeilToInt(Stage.GetEndTimeCode());
 			UsdUtils::GetAnimatedMeshTimeCodes(Stage, PrimPath, StartFrame, EndFrame);
@@ -2044,33 +2054,26 @@ namespace UE::InterchangeUsdTranslator::Private
 				const FString& SlotName = Slot.MaterialSource;
 
 				// Get the Uid of the material instance that we'll end up assigning to this slot
-				FString MaterialInstanceUid;
-				bool bAddDirectlyAsSlotDependency = false;
+				FString MaterialUid;
 				switch (Slot.AssignmentType)
 				{
 					case UsdUtils::EPrimAssignmentType::DisplayColor:
 					{
-						bAddDirectlyAsSlotDependency = true;
-
 						// MaterialSource here is e.g. "!DisplayColor_0_1"
-						MaterialInstanceUid = AddDisplayColorMaterialInstanceNodeIfNeeded(NodeContainer, Slot);
+						MaterialUid = AddDisplayColorMaterialInstanceNodeIfNeeded(NodeContainer, Slot);
 						break;
 					}
 					case UsdUtils::EPrimAssignmentType::MaterialPrim:
 					{
-						MaterialInstanceUid = MaterialPrefix + Slot.MaterialSource;	   // This is the prim path
-						if (Slot.bMeshIsDoubleSided)
-						{
-							MaterialInstanceUid += TwoSidedSuffix;
-						}
+						// MaterialSource here is the material prim path
+						UE::FUsdPrim MaterialPrim = Stage.GetPrimAtPath(UE::FSdfPath{*Slot.MaterialSource});
+						MaterialUid = AddMaterialNode(MaterialPrim, TranslatorSettings, NodeContainer, Slot.bMeshIsDoubleSided);
 						break;
 					}
 					case UsdUtils::EPrimAssignmentType::UnrealMaterial:
 					{
-						bAddDirectlyAsSlotDependency = true;
-
 						// MaterialSource here is the content path, e.g. "/Game/MyFolder/Red.Red"
-						MaterialInstanceUid = AddUnrealMaterialReferenceNodeIfNeeded(NodeContainer, Slot.MaterialSource);
+						MaterialUid = AddUnrealMaterialReferenceNodeIfNeeded(NodeContainer, Slot.MaterialSource);
 						break;
 					}
 					default:
@@ -2080,22 +2083,11 @@ namespace UE::InterchangeUsdTranslator::Private
 					}
 				}
 
-				// If we're a DisplayColor/Unreal material, we know everything we need right now
-				if (bAddDirectlyAsSlotDependency)
-				{
-					MeshNode->SetSlotMaterialDependencyUid(SlotName, *MaterialInstanceUid);
-				}
-				// If we found a match let's set the slot to the corresponding Material right away, as we already must have traversed this material
-				else if (FString* ActualMaterialInstanceUid = MaterialUidToActualNodeUid.Find(MaterialInstanceUid))
-				{
-					MeshNode->SetSlotMaterialDependencyUid(SlotName, *ActualMaterialInstanceUid);
-				}
-				// Otherwise, we need to wait until the material prim itself is translated, as we may need to defer to another translator
-				// (e.g. MaterialX) for the translation, which could generate an entirely different translated node we can't know about yet
-				else
+				MeshNode->SetSlotMaterialDependencyUid(SlotName, *MaterialUid);
+
+				if (TArray<FUsdMaterialXShaderGraph::FGeomProp>* FoundGeomProps = MaterialUidToGeomProps.Find(MaterialUid))
 				{
-					// one material can be attached to several meshes
-					PrimPathToSlotMeshNodes.FindOrAdd(MaterialInstanceUid).Add({SlotName, MeshNode});
+					AddPrimvarBakingAttributes(MeshNode, MaterialUid, NodeContainer, *FoundGeomProps);
 				}
 			}
 		}
@@ -3295,11 +3287,11 @@ namespace UE::InterchangeUsdTranslator::Private
 		TArray<FString> AssetNodeUids;
 		if (Prim.IsA(TEXT("Material")))
 		{
-			AssetNodeUids = TranslatorImpl.AddMaterialNode(Prim, TranslatorSettings, NodeContainer);
+			AssetNodeUids = {TranslatorImpl.AddMaterialNode(Prim, TranslatorSettings, NodeContainer)};
 		}
 		else if (Prim.IsA(TEXT("Mesh")))
 		{
-			AssetNodeUids = {TranslatorImpl.AddMeshNode(Prim, NodeContainer, Info)};
+			AssetNodeUids = {TranslatorImpl.AddMeshNode(Prim, TranslatorSettings, NodeContainer, Info)};
 		}
 		else if (Prim.IsA(TEXT("Camera")))
 		{
@@ -3315,7 +3307,7 @@ namespace UE::InterchangeUsdTranslator::Private
 			// Volumes are also not currently supported.
 
 			constexpr bool bPrimitiveShape_True = true;
-			AssetNodeUids = {TranslatorImpl.AddMeshNode(Prim, NodeContainer, Info, bPrimitiveShape_True)};
+			AssetNodeUids = {TranslatorImpl.AddMeshNode(Prim, TranslatorSettings, NodeContainer, Info, bPrimitiveShape_True)};
 		}
 		else if (Prim.IsA(TEXT("Volume")))
 		{
diff --git a/Engine/Plugins/Interchange/Runtime/Source/Common/Public/Usd/InterchangeUsdDefinitions.h b/Engine/Plugins/Interchange/Runtime/Source/Common/Public/Usd/InterchangeUsdDefinitions.h
--- a/Engine/Plugins/Interchange/Runtime/Source/Common/Public/Usd/InterchangeUsdDefinitions.h
+++ b/Engine/Plugins/Interchange/Runtime/Source/Common/Public/Usd/InterchangeUsdDefinitions.h
@@ -11,9 +11,6 @@ namespace UE::Interchange::USD
 	// Name of a custom attribute added to translated nodes to contain their geometry purpose (proxy, render, guide, etc.)
 	const FString GeometryPurposeIdentifier = TEXT("USD_Geometry_Purpose");
 
-	// Prefix for custom attributes that describe material instance parameters
-	const FString AttributeIdentifierPrefix = TEXT("USD_MI_Attribute_");
-
 	// Prefixes we use to stash some primvar mapping information as custom attributes on mesh / material nodes, so that
 	// the USD Pipeline can produce primvar-compatible materials
 	const FString PrimvarUVIndexAttributePrefix = TEXT("USD_PrimvarUVIndex_");
@@ -28,24 +25,10 @@ namespace UE::Interchange::USD
 	const FString UseTextureParameterSuffix = TEXT("Texture");
 	const FString UVIndexParameterSuffix = TEXT("UVIndex");
 
-	// Flag indicating whether we should parse a UInterchangeShaderGraphNode on the InterchangeUSDPipeline
+	// Flag indicating whether we should parse a UInterchangeShaderGraphNode on the InterchangeUSDPipeline.
+	// This is now only used for the compatible primvar code
 	const FString ParseMaterialIdentifier = TEXT("USD_MI_ParseMaterial");
 
-	// Means we should pick this material to use as parent material, regardless of what else is provided
-	// (NeedsVTParent, DisplayColor, etc.)
-	const FString ParentMaterialIdentifier = TEXT("USD_MI_Parent");
-
-	// Means that we should use a display color material parent for this material, according to the
-	// value provided in this attribute
-	const FString DisplayColorDescIdentifier = TEXT("USD_MI_DisplayColorDescription");
-
-	// Means that if we were to pick one of the UsdPreviewSurface parent materials for this material
-	// translated node, we should pick the "VT" variant of that material.
-	//
-	// Note: To indicate TwoSided/Translucent we use the CustomTwoSided/BlendMode standard UInterchangeShaderGraphNode
-	// attributes.
-	const FString NeedsVTParentIdentifier = TEXT("USD_MI_NeedsVTParent");
-
 	// Used for volumetric material parameters, whenever we assign an SVT to a material as a fallback due
 	// to it's field name only
 	const FString VolumeFieldNameMaterialParameterPrefix = TEXT("USD_FieldName_");
