Interchange USD: Now that Interchange USD is int its own plugin, the cvar to enable USD through Interchange is true by default.

diff --git a/Engine/Plugins/Interchange/Extensions/OpenUSD/Source/Import/Private/InterchangeUsdTranslator.cpp b/Engine/Plugins/Interchange/Extensions/OpenUSD/Source/Import/Private/InterchangeUsdTranslator.cpp
--- a/Engine/Plugins/Interchange/Extensions/OpenUSD/Source/Import/Private/InterchangeUsdTranslator.cpp
+++ b/Engine/Plugins/Interchange/Extensions/OpenUSD/Source/Import/Private/InterchangeUsdTranslator.cpp
@@ -81,7 +81,7 @@
 
 #define LOCTEXT_NAMESPACE "InterchangeUSDTranslator"
 
-static bool GInterchangeEnableUSDImport = false;
+static bool GInterchangeEnableUSDImport = true;
 static FAutoConsoleVariableRef CVarInterchangeEnableUSDImport(
 	TEXT("Interchange.FeatureFlags.Import.USD"),
 	GInterchangeEnableUSDImport,
