Interchange USD: The Interchange USD pipeline stacks are now registered by the Interchange USD plugin at startup.

diff --git a/Engine/Plugins/Interchange/Extensions/OpenUSD/Source/Import/Private/InterchangeOpenUSDImportModule.cpp b/Engine/Plugins/Interchange/Extensions/OpenUSD/Source/Import/Private/InterchangeOpenUSDImportModule.cpp
--- a/Engine/Plugins/Interchange/Extensions/OpenUSD/Source/Import/Private/InterchangeOpenUSDImportModule.cpp
+++ b/Engine/Plugins/Interchange/Extensions/OpenUSD/Source/Import/Private/InterchangeOpenUSDImportModule.cpp
@@ -4,6 +4,7 @@
 
 #include "Engine/Engine.h"
 #include "InterchangeManager.h"
+#include "InterchangeProjectSettings.h"
 #include "InterchangeUsdTranslator.h"
 #include "Misc/CoreDelegates.h"
 #include "Modules/ModuleManager.h"
@@ -15,18 +16,39 @@ void FInterchangeOpenUSDImportModule::StartupModule()
 	{
 		UInterchangeManager& InterchangeManager = UInterchangeManager::GetInterchangeManager();
 		InterchangeManager.RegisterTranslator(UInterchangeUSDTranslator::StaticClass());
+
+		// Add USD translator and pipelines to the Interchange engine
+		{
+			FInterchangeTranslatorPipelines TranslatorPipelines;
+			TranslatorPipelines.Translator = TSoftClassPtr<UInterchangeTranslatorBase>(UInterchangeUSDTranslator::StaticClass());
+			TranslatorPipelines.Pipelines.Add(FSoftObjectPath(TEXT("/InterchangeOpenUSD/Pipelines/DefaultUSDAssetsPipeline.DefaultUSDAssetsPipeline")));
+			TranslatorPipelines.Pipelines.Add(FSoftObjectPath(TEXT("/InterchangeOpenUSD/Pipelines/DefaultUSDPipelineAssetImport.DefaultUSDPipelineAssetImport")));
+			TranslatorPipelines.Pipelines.Add(FSoftObjectPath(TEXT("/Interchange/Pipelines/DefaultMaterialXPipeline.DefaultMaterialXPipeline")));
+
+			FInterchangeImportSettings& ImportSettings = FInterchangeProjectSettingsUtils::GetMutableDefaultImportSettings(false);
+			FInterchangePipelineStack& PipelineStack = ImportSettings.PipelineStacks[TEXT("Assets")];
+			PipelineStack.PerTranslatorPipelines.Add(TranslatorPipelines);
+		}
+
+		{
+			FInterchangeTranslatorPipelines TranslatorPipelines;
+			TranslatorPipelines.Translator = TSoftClassPtr<UInterchangeTranslatorBase>(UInterchangeUSDTranslator::StaticClass());
+			TranslatorPipelines.Pipelines.Add(FSoftObjectPath(TEXT("/InterchangeOpenUSD/Pipelines/DefaultUSDSceneAssetsPipeline.DefaultUSDSceneAssetsPipeline")));
+			TranslatorPipelines.Pipelines.Add(FSoftObjectPath(TEXT("/Interchange/Pipelines/DefaultSceneLevelPipeline.DefaultSceneLevelPipeline")));
+			TranslatorPipelines.Pipelines.Add(FSoftObjectPath(TEXT("/InterchangeOpenUSD/Pipelines/DefaultUSDPipeline.DefaultUSDPipeline")));
+			TranslatorPipelines.Pipelines.Add(FSoftObjectPath(TEXT("/Interchange/Pipelines/DefaultMaterialXPipeline.DefaultMaterialXPipeline")));
+
+			FInterchangeImportSettings& ImportSettings = FInterchangeProjectSettingsUtils::GetMutableDefaultImportSettings(true);
+			FInterchangePipelineStack& PipelineStack = ImportSettings.PipelineStacks[TEXT("Scene")];
+			PipelineStack.PerTranslatorPipelines.Add(TranslatorPipelines);
+		}
 	};
-	if (GEngine)
-	{
-		RegisterItems();
-	}
-	else
-	{
-		FCoreDelegates::OnPostEngineInit.AddLambda(RegisterItems);
-	}
+
+	PostEngineInitHandle = FCoreDelegates::OnPostEngineInit.AddLambda(RegisterItems);
 }
 void FInterchangeOpenUSDImportModule::ShutdownModule()
 {
+	FCoreDelegates::OnPostEngineInit.Remove(PostEngineInitHandle);
 }
 
 IMPLEMENT_MODULE(FInterchangeOpenUSDImportModule, InterchangeOpenUSDImport)
\ No newline at end of file
diff --git a/Engine/Plugins/Interchange/Extensions/OpenUSD/Source/Import/Public/InterchangeOpenUSDImportModule.h b/Engine/Plugins/Interchange/Extensions/OpenUSD/Source/Import/Public/InterchangeOpenUSDImportModule.h
--- a/Engine/Plugins/Interchange/Extensions/OpenUSD/Source/Import/Public/InterchangeOpenUSDImportModule.h
+++ b/Engine/Plugins/Interchange/Extensions/OpenUSD/Source/Import/Public/InterchangeOpenUSDImportModule.h
@@ -3,10 +3,14 @@
 #pragma once
 
 #include "Modules/ModuleInterface.h"
+#include "Delegates/IDelegateInstance.h"
 
 class FInterchangeOpenUSDImportModule : public IModuleInterface
 {
 public:
 	virtual void StartupModule() override;
 	virtual void ShutdownModule() override;
+
+private:
+	FDelegateHandle PostEngineInitHandle;
 };
\ No newline at end of file
