USD: Fix logic bug when removing factory nodes due to geometry purpose.

diff --git a/Engine/Plugins/Interchange/Extensions/OpenUSD/Source/Import/Private/InterchangeUSDPipeline.cpp b/Engine/Plugins/Interchange/Extensions/OpenUSD/Source/Import/Private/InterchangeUSDPipeline.cpp
--- a/Engine/Plugins/Interchange/Extensions/OpenUSD/Source/Import/Private/InterchangeUSDPipeline.cpp
+++ b/Engine/Plugins/Interchange/Extensions/OpenUSD/Source/Import/Private/InterchangeUSDPipeline.cpp
@@ -133,28 +133,22 @@ namespace UE::InterchangeUsdPipeline::Private
 					TArray<FString> TargetTranslatedNodeUids;
 					FactoryNode->GetTargetNodeUids(TargetTranslatedNodeUids);
 
-					// If we have a single target node that is disabled, just disable the factory node itself.
-					// This is probably the general case, and should be more explicit on preview panels and such
-					if (TargetTranslatedNodeUids.Num() == 1)
+					bool bRemovedNode = false;
+					for (const FString& TargetNode : TargetTranslatedNodeUids)
 					{
-						if (DisabledTranslatedNodes.Contains(TargetTranslatedNodeUids[0]))
+						if (DisabledTranslatedNodes.Contains(TargetNode))
 						{
-							const bool bEnabled = false;
-							FactoryNode->SetEnabled(bEnabled);
+							bRemovedNode = true;
+							FactoryNode->RemoveTargetNodeUid(TargetNode);
 						}
 					}
-					// If we have multiple target nodes, instead of disabling the whole factory node just remove
-					// this one particular target (e.g. removing a single mesh from a collapsed mesh blob, in case that
-					// mesh doesn't match the purpose)
-					else
+
+					// We don't want to leave a factory node enabled without a target, as that may lead
+					// to some errors on the factories
+					if (bRemovedNode && FactoryNode->GetTargetNodeCount() == 0)
 					{
-						for (const FString& TargetNode : TargetTranslatedNodeUids)
-						{
-							if (DisabledTranslatedNodes.Contains(TargetNode))
-							{
-								FactoryNode->RemoveTargetNodeUid(TargetNode);
-							}
-						}
+						const bool bEnabled = false;
+						FactoryNode->SetEnabled(bEnabled);
 					}
 				}
 
@@ -165,23 +159,20 @@ namespace UE::InterchangeUsdPipeline::Private
 					TArray<FString> TranslatedMeshNodeUids;
 					LodDataNode->GetMeshUids(TranslatedMeshNodeUids);
 
-					if (TranslatedMeshNodeUids.Num() == 1)
+					bool bRemovedNode = false;
+					for (const FString& MeshNodeUid : TranslatedMeshNodeUids)
 					{
-						if (DisabledTranslatedNodes.Contains(TranslatedMeshNodeUids[0]))
+						if (DisabledTranslatedNodes.Contains(MeshNodeUid))
 						{
-							const bool bEnabled = false;
-							LodDataNode->SetEnabled(bEnabled);
+							bRemovedNode = true;
+							LodDataNode->RemoveMeshUid(MeshNodeUid);
 						}
 					}
-					else
+
+					if (bRemovedNode && LodDataNode->GetMeshUidsCount() == 0)
 					{
-						for (const FString& MeshNodeUid : TranslatedMeshNodeUids)
-						{
-							if (DisabledTranslatedNodes.Contains(MeshNodeUid))
-							{
-								LodDataNode->RemoveMeshUid(MeshNodeUid);
-							}
-						}
+						const bool bEnabled = false;
+						LodDataNode->SetEnabled(bEnabled);
 					}
 				}
 				else if (UInterchangeSkeletalMeshLodDataNode* SkeletalNode = Cast<UInterchangeSkeletalMeshLodDataNode>(FactoryNode))
@@ -189,23 +180,20 @@ namespace UE::InterchangeUsdPipeline::Private
 					TArray<FString> TranslatedMeshNodeUids;
 					SkeletalNode->GetMeshUids(TranslatedMeshNodeUids);
 
-					if (TranslatedMeshNodeUids.Num() == 1)
+					bool bRemovedNode = false;
+					for (const FString& MeshNodeUid : TranslatedMeshNodeUids)
 					{
-						if (DisabledTranslatedNodes.Contains(TranslatedMeshNodeUids[0]))
+						if (DisabledTranslatedNodes.Contains(MeshNodeUid))
 						{
-							const bool bEnabled = false;
-							SkeletalNode->SetEnabled(bEnabled);
+							bRemovedNode = true;
+							SkeletalNode->RemoveMeshUid(MeshNodeUid);
 						}
 					}
-					else
+
+					if (bRemovedNode && SkeletalNode->GetMeshUidsCount() == 0)
 					{
-						for (const FString& MeshNodeUid : TranslatedMeshNodeUids)
-						{
-							if (DisabledTranslatedNodes.Contains(MeshNodeUid))
-							{
-								SkeletalNode->RemoveMeshUid(MeshNodeUid);
-							}
-						}
+						const bool bEnabled = false;
+						SkeletalNode->SetEnabled(bEnabled);
 					}
 				}
 			}
