USD: Prevent render contexts registered with legacy USD from showing up on Interchange USD's list.

diff --git a/Engine/Plugins/Interchange/Extensions/OpenUSD/Source/Editor/Private/InterchangeUsdTranslatorSettingsCustomization.cpp b/Engine/Plugins/Interchange/Extensions/OpenUSD/Source/Editor/Private/InterchangeUsdTranslatorSettingsCustomization.cpp
--- a/Engine/Plugins/Interchange/Extensions/OpenUSD/Source/Editor/Private/InterchangeUsdTranslatorSettingsCustomization.cpp
+++ b/Engine/Plugins/Interchange/Extensions/OpenUSD/Source/Editor/Private/InterchangeUsdTranslatorSettingsCustomization.cpp
@@ -48,9 +48,18 @@ void FInterchangeUsdTranslatorSettingsCustomization::CustomizeDetails(IDetailLay
 		return;
 	}
 
+	// We used to pull these from UsdUnreal::MaterialUtils::GetRegisteredRenderContexts(), but these are
+	// hard-coded for Interchange because that list may contain render contexts added in through the
+	// legacy USD render context registry, which we don't necessarily support on Interchange just yet (e.g. "mdl")
+	TArray<FName> AllowedRenderContexts = {
+		TEXT("mtlx"),
+		TEXT("universal"),
+		TEXT("unreal"),
+	};
+
 	RenderContextComboBoxItems.Reset();
 	TSharedPtr<FString> InitiallySelectedContext;
-	for (const FName& Context : UsdUnreal::MaterialUtils::GetRegisteredRenderContexts())
+	for (const FName& Context : AllowedRenderContexts)
 	{
 		TSharedPtr<FString> ContextStr;
 		if (Context == UnrealIdentifiers::UniversalRenderContext)
