[USD Interchange] Adds check for a valid and active prim when attempting to retrieve a mesh lod from the stage. This prevents a crash if a primitive at the assumed path does not exist.

diff --git a/Engine/Plugins/Interchange/Extensions/OpenUSD/Source/Import/Private/InterchangeUsdTranslator.cpp b/Engine/Plugins/Interchange/Extensions/OpenUSD/Source/Import/Private/InterchangeUsdTranslator.cpp
--- a/Engine/Plugins/Interchange/Extensions/OpenUSD/Source/Import/Private/InterchangeUsdTranslator.cpp
+++ b/Engine/Plugins/Interchange/Extensions/OpenUSD/Source/Import/Private/InterchangeUsdTranslator.cpp
@@ -436,7 +436,7 @@ namespace UE::InterchangeUsdTranslator::Private
 		const UE::FSdfPath IdealMeshPrimPath = LODContainerPrim.GetPrimPath().AppendChild(*LODName);
 
 		UE::FUsdPrim Prim = LODContainerPrim.GetStage().GetPrimAtPath(IdealMeshPrimPath);
-		if (Prim.IsA(TEXT("Mesh")))
+		if (Prim && Prim.IsActive() && Prim.IsA(TEXT("Mesh")))
 		{
 			return Prim;
 		}
