Interchange USD:
- Add an option to allow users to choose how to import primvars

diff --git a/Engine/Plugins/Interchange/Extensions/OpenUSD/Source/Import/InterchangeOpenUSDImport.build.cs b/Engine/Plugins/Interchange/Extensions/OpenUSD/Source/Import/InterchangeOpenUSDImport.build.cs
--- a/Engine/Plugins/Interchange/Extensions/OpenUSD/Source/Import/InterchangeOpenUSDImport.build.cs
+++ b/Engine/Plugins/Interchange/Extensions/OpenUSD/Source/Import/InterchangeOpenUSDImport.build.cs
@@ -14,6 +14,7 @@ namespace UnrealBuildTool.Rules
 					"Core",
 					"CoreUObject",
 					"Engine",
+					"InterchangeCommon",
 					"InterchangeCommonParser",	// For FAnimationPayloadData
 					"InterchangeCore", 			// For UInterchangeBaseNode, base translator
 					"InterchangeImport", 		// For the translator settings, pipeline interface, etc.
diff --git a/Engine/Plugins/Interchange/Extensions/OpenUSD/Source/Import/Private/InterchangeUSDPipeline.cpp b/Engine/Plugins/Interchange/Extensions/OpenUSD/Source/Import/Private/InterchangeUSDPipeline.cpp
--- a/Engine/Plugins/Interchange/Extensions/OpenUSD/Source/Import/Private/InterchangeUSDPipeline.cpp
+++ b/Engine/Plugins/Interchange/Extensions/OpenUSD/Source/Import/Private/InterchangeUSDPipeline.cpp
@@ -17,7 +17,6 @@
 #include "InterchangeVolumeNode.h"
 #include "Nodes/InterchangeBaseNodeContainer.h"
 #include "Nodes/InterchangeUserDefinedAttribute.h"
-#include "Usd/InterchangeUsdDefinitions.h"
 
 #include "DynamicMesh/DynamicMesh3.h"
 #include "Engine/StaticMesh.h"
@@ -563,6 +562,7 @@ namespace UE::InterchangeUsdPipeline::Private
 #if WITH_EDITORONLY_DATA
 	/**
 	 * Add the primvar attributes from the mesh nodes to factory nodes
+	 * Also add the array names as payload key attributes to filter them in the Translator
 	 */
 	void AddPrimvarAttributesToFactoryNodes(
 		const UInterchangeMeshNode* MeshNode,
@@ -575,12 +575,14 @@ namespace UE::InterchangeUsdPipeline::Private
 		if (int32 NumPrimvar; MeshNode->GetInt32Attribute(USD::Primvar::Number, NumPrimvar))
 		{
 			MeshFactoryNode->AddInt32Attribute(USD::Primvar::Number, NumPrimvar);
+			MeshFactoryNode->AddPayloadKeyInt32Attribute(USD::Primvar::Number, NumPrimvar);
 			for (int32 Index = 0; Index < NumPrimvar; ++Index)
 			{
 				FString PrimvarAttribute = USD::Primvar::Name + FString::FromInt(Index);
 				if (FString PrimvarName; MeshNode->GetStringAttribute(PrimvarAttribute, PrimvarName))
 				{
 					MeshFactoryNode->AddStringAttribute(PrimvarAttribute, PrimvarName);
+					MeshFactoryNode->AddPayloadKeyStringAttribute(PrimvarAttribute, PrimvarName);
 				}
 
 				if (bool bTangentSpace; MeshNode->GetBooleanAttribute(PrimvarAttribute, bTangentSpace))
@@ -951,7 +953,7 @@ namespace UE::InterchangeUsdPipeline::Private
 		}
 	}
 
-	void ProcessMeshNodes(UInterchangeBaseNodeContainer* NodeContainer, bool bGenerateCompatibleMaterials)
+	void ProcessMeshNodes(UInterchangeBaseNodeContainer* NodeContainer, bool bGenerateCompatibleMaterials, EInterchangeUsdPrimvar ImportPrimvar)
 	{
 		using namespace UE::Interchange;
 
@@ -960,7 +962,7 @@ namespace UE::InterchangeUsdPipeline::Private
 		TMap<const UInterchangeMeshNode*, UInterchangeMeshFactoryNode*> MeshNodeToFactoryNode;
 
 		NodeContainer->IterateNodesOfType<UInterchangeMeshNode>(
-			[NodeContainer, &MeshNodeToFactoryNode](const FString& NodeUid, UInterchangeMeshNode* MeshNode)
+			[NodeContainer, &MeshNodeToFactoryNode, ImportPrimvar](const FString& NodeUid, UInterchangeMeshNode* MeshNode)
 			{
 				const FString FactoryNodeUid = UInterchangeFactoryBaseNode::BuildFactoryNodeUid(NodeUid);
 				UInterchangeMeshFactoryNode* MeshFactoryNode = Cast<UInterchangeMeshFactoryNode>(NodeContainer->GetFactoryNode(FactoryNodeUid));
@@ -969,10 +971,15 @@ namespace UE::InterchangeUsdPipeline::Private
 					return;
 				}
 
+				MeshFactoryNode->AddPayloadKeyInt32Attribute(UE::Interchange::USD::Primvar::Import, int32(ImportPrimvar));
+
 				MeshNodeToFactoryNode.Add(MeshNode, MeshFactoryNode);
 
 #if WITH_EDITORONLY_DATA
-				AddPrimvarAttributesToFactoryNodes(MeshNode, MeshFactoryNode, NodeContainer);
+				if(ImportPrimvar != EInterchangeUsdPrimvar::Standard)
+				{
+					AddPrimvarAttributesToFactoryNodes(MeshNode, MeshFactoryNode, NodeContainer);
+				}
 #endif	  // WITH_EDITORONLY_DATA
 			}
 		);
@@ -1287,9 +1294,12 @@ void UInterchangeUsdPipeline::ExecutePostFactoryPipeline(
 
 	if (UStaticMesh* StaticMeshAsset = Cast<UStaticMesh>(CreatedAsset))
 	{
-		using namespace UE::InterchangeUsdPipeline::Private;
-		FUsdPrimvarBaker Baker(StaticMeshAsset, BaseNodeContainer, FactoryNode);
-		Baker.Bake();
+		if(ImportPrimvars != EInterchangeUsdPrimvar::Standard)
+		{
+			using namespace UE::InterchangeUsdPipeline::Private;
+			FUsdPrimvarBaker Baker(StaticMeshAsset, BaseNodeContainer, FactoryNode);
+			Baker.Bake();
+		}
 	}
 #endif	  // USE_USD_SDK && WITH_EDITORONLY_DATA
 }
@@ -1321,7 +1331,7 @@ void UInterchangeUsdPipeline::ExecutePipeline(
 		RemovePseudoRootFactoryNode(*NodeContainer);
 	}
 
-	ProcessMeshNodes(NodeContainer, bGeneratePrimvarCompatibleMaterials);
+	ProcessMeshNodes(NodeContainer, bGeneratePrimvarCompatibleMaterials, ImportPrimvars);
 #endif	  // USE_USD_SDK
 }
 
diff --git a/Engine/Plugins/Interchange/Extensions/OpenUSD/Source/Import/Private/InterchangeUsdTranslator.cpp b/Engine/Plugins/Interchange/Extensions/OpenUSD/Source/Import/Private/InterchangeUsdTranslator.cpp
--- a/Engine/Plugins/Interchange/Extensions/OpenUSD/Source/Import/Private/InterchangeUsdTranslator.cpp
+++ b/Engine/Plugins/Interchange/Extensions/OpenUSD/Source/Import/Private/InterchangeUsdTranslator.cpp
@@ -4528,12 +4528,35 @@ TOptional<UE::Interchange::FMeshPayloadData> UInterchangeUSDTranslator::GetMeshP
 		return Result;
 	}
 	using namespace UE::Interchange;
+	using namespace UsdUtils;
+	using namespace UsdToUnreal;
 
 	UsdToUnreal::FUsdMeshConversionOptions OptionsCopy = ImplPtr->CachedMeshConversionOptions;
 	PayloadAttributes.GetAttribute(UE::Interchange::FAttributeKey{ MeshPayload::Attributes::MeshGlobalTransform }, OptionsCopy.AdditionalTransform);
 	OptionsCopy.bMergeIdenticalMaterialSlots = false;	 // This must always be false, because we need the material assignments we read from the
 														 // meshes to match up with whatever we cached from AddMeshNode, in order to fixup LOD
 														 // material slots
+		
+	static_assert(uint8(EInterchangeUsdPrimvar::All) == FUsdMeshConversionOptions::EImportPrimvar::All, "FUsdMeshConversionOptions::EImportPrimvar::All is different from EInterchangeUsdPrimvar::All");
+	static_assert(uint8(EInterchangeUsdPrimvar::Bake) == FUsdMeshConversionOptions::EImportPrimvar::Bake, "FUsdMeshConversionOptions::EImportPrimvar::Bake is different from EInterchangeUsdPrimvar::Bake");
+	static_assert(uint8(EInterchangeUsdPrimvar::Standard) == FUsdMeshConversionOptions::EImportPrimvar::Standard, "FUsdMeshConversionOptions::EImportPrimvar::Standard is different from EInterchangeUsdPrimvar::Standard");
+
+	PayloadAttributes.GetAttribute(UE::Interchange::FAttributeKey{ USD::Primvar::Import }, reinterpret_cast<int32&>(OptionsCopy.ImportPrimvars));
+
+	if (EInterchangeUsdPrimvar(OptionsCopy.ImportPrimvars) != EInterchangeUsdPrimvar::Standard)
+	{
+		if (int32 PrimvarNumber; PayloadAttributes.GetAttribute(UE::Interchange::FAttributeKey{ USD::Primvar::Number }, PrimvarNumber) == EAttributeStorageResult::Operation_Success)
+		{
+			OptionsCopy.PrimvarNames.Reserve(PrimvarNumber);
+			for (int32 Index = 0; Index < PrimvarNumber; ++Index)
+			{
+				if (FString PrimvarName; PayloadAttributes.GetAttribute(UE::Interchange::FAttributeKey{ USD::Primvar::Name + FString::FromInt(Index) }, PrimvarName) == EAttributeStorageResult::Operation_Success)
+				{
+					OptionsCopy.PrimvarNames.Emplace(MoveTemp(PrimvarName));
+				}
+			}
+		}
+	}
 
 	UE::Interchange::FMeshPayloadData MeshPayloadData;
 	switch (PayloadKey.Type)
diff --git a/Engine/Plugins/Interchange/Extensions/OpenUSD/Source/Import/Public/InterchangeUSDPipeline.h b/Engine/Plugins/Interchange/Extensions/OpenUSD/Source/Import/Public/InterchangeUSDPipeline.h
--- a/Engine/Plugins/Interchange/Extensions/OpenUSD/Source/Import/Public/InterchangeUSDPipeline.h
+++ b/Engine/Plugins/Interchange/Extensions/OpenUSD/Source/Import/Public/InterchangeUSDPipeline.h
@@ -4,6 +4,7 @@
 
 #include "InterchangePipelineBase.h"
 #include "InterchangeSourceData.h"
+#include "Usd/InterchangeUsdDefinitions.h"
 
 #include "InterchangeUSDPipeline.generated.h"
 
@@ -29,6 +30,10 @@ public:
 	UPROPERTY(EditAnywhere, BlueprintReadWrite, Category = "USD", meta = (Bitmask, BitmaskEnum = "/Script/UnrealUSDWrapper.EUsdPurpose"))
 	int32 GeometryPurpose;
 
+	/** Setting to tell what primvars have to be attached to the MeshDescription */
+	UPROPERTY(EditAnywhere, BlueprintReadWrite, Category = "USD")
+	EInterchangeUsdPrimvar ImportPrimvars = EInterchangeUsdPrimvar::Bake;
+
 	/**
 	 * The translator always emits a scene node for the stage pseudoroot. Setting this option to false means we will remove it from
 	 * the factory nodes, making top-level prims lead to top-level actors and components.
diff --git a/Engine/Plugins/Interchange/Runtime/Source/Common/Public/Usd/InterchangeUsdDefinitions.h b/Engine/Plugins/Interchange/Runtime/Source/Common/Public/Usd/InterchangeUsdDefinitions.h
--- a/Engine/Plugins/Interchange/Runtime/Source/Common/Public/Usd/InterchangeUsdDefinitions.h
+++ b/Engine/Plugins/Interchange/Runtime/Source/Common/Public/Usd/InterchangeUsdDefinitions.h
@@ -4,6 +4,21 @@
 
 #include "CoreMinimal.h"
 
+#include "InterchangeUsdDefinitions.generated.h"
+
+UENUM(BlueprintType)
+enum class EInterchangeUsdPrimvar : uint8
+{
+	/** Store only the standard primvars such as UVs, VertexColors, etc.*/
+	Standard = 0,
+
+	/** Store only primvars in the Mesh Description used for baking to textures (basically <geompropvalue> node from MaterialX shadergraphs that are converted to <image>)*/
+	Bake,
+
+	/** Store all primvars in the MeshDescription, regardless*/
+	All
+};
+
 namespace UE::Interchange::USD
 {
 	const FString USDContextTag = TEXT("USD");
@@ -58,5 +73,7 @@ namespace UE::Interchange::USD
 		const FString TangentSpace = TEXT("USD_PrimvarTangentSpace");
 		// Attribute that informs about the UID of a ShaderNode of type TextureSample
 		const FString ShaderNodeTextureSample = TEXT("USD_ShaderNodeTextureSample");
+		// Attribute that informs how we should handle primvars on MeshDescriptions
+		const FString Import = TEXT("USD_Import_Primvars");
 	}
 }	 // namespace UE::Interchange::USD
diff --git a/Engine/Plugins/Runtime/USDCore/Source/USDUtilities/Private/USDGeomMeshConversion.cpp b/Engine/Plugins/Runtime/USDCore/Source/USDUtilities/Private/USDGeomMeshConversion.cpp
--- a/Engine/Plugins/Runtime/USDCore/Source/USDUtilities/Private/USDGeomMeshConversion.cpp
+++ b/Engine/Plugins/Runtime/USDCore/Source/USDUtilities/Private/USDGeomMeshConversion.cpp
@@ -1304,6 +1304,7 @@ namespace UE::UsdGeomMeshConversion::Private
 
 	void CollectPrimvarData(const pxr::UsdPrim& Prim, const UsdToUnreal::FUsdMeshConversionOptions& Options, FUsdMeshData& InOutMeshData)
 	{
+		using namespace UsdToUnreal;
 		const pxr::TfToken DefaultInterpolation = pxr::UsdGeomTokens->varying;
 
 		pxr::UsdGeomPrimvarsAPI PrimvarAPI{Prim};
@@ -1311,7 +1312,8 @@ namespace UE::UsdGeomMeshConversion::Private
 		for (const pxr::UsdGeomPrimvar& Primvar : Primvars)
 		{
 			std::string PrimvarName = Primvar.GetPrimvarName().GetString();
-			if (PrimvarName == "displayColor" || PrimvarName == "displayOpacity" || PrimvarName == "st")
+			if (PrimvarName == "displayColor" || PrimvarName == "displayOpacity" || PrimvarName == "st"
+				|| (Options.ImportPrimvars == FUsdMeshConversionOptions::EImportPrimvar::Bake && !Options.PrimvarNames.Contains(UTF8_TO_TCHAR(PrimvarName.c_str()))))
 			{
 				continue;
 			}
@@ -1497,6 +1499,7 @@ namespace UE::UsdGeomMeshConversion::Private
 			}
 
 			// Custom primvars
+			if (Options.ImportPrimvars != UsdToUnreal::FUsdMeshConversionOptions::EImportPrimvar::Standard)
 			{
 				CollectPrimvarData(Prim, Options, InOutMeshData);
 			}
diff --git a/Engine/Plugins/Runtime/USDCore/Source/USDUtilities/Public/USDGeomMeshConversion.h b/Engine/Plugins/Runtime/USDCore/Source/USDUtilities/Public/USDGeomMeshConversion.h
--- a/Engine/Plugins/Runtime/USDCore/Source/USDUtilities/Public/USDGeomMeshConversion.h
+++ b/Engine/Plugins/Runtime/USDCore/Source/USDUtilities/Public/USDGeomMeshConversion.h
@@ -56,6 +56,19 @@ namespace UsdToUnreal
 	/** Common options to mesh conversion functions */
 	struct USDUTILITIES_API FUsdMeshConversionOptions
 	{
+		/** This enum must match EInterchangeUsdPrimvar in InterchangeUsdDefinitions.h*/
+		enum EImportPrimvar
+		{
+			/** Store only the standard primvars such as UVs, VertexColors, etc.*/
+			Standard = 0,
+
+			/** Store only primvars in the Mesh Description used for baking to textures (basically <geompropvalue> node from MaterialX shadergraphs that are converted to <image>)*/
+			Bake,
+
+			/** Store all primvars in the MeshDescription, regardless*/
+			All
+		};
+
 		// These need to be explicitly defaulted surrounded by the deprecation macros or else the implicit versions will
 		// be generated elsewhere, use MaterialToPrimvarToUVIndex and trigger the deprecation warnings
 		PRAGMA_DISABLE_DEPRECATION_WARNINGS
@@ -66,6 +79,9 @@ namespace UsdToUnreal
 		const static FUsdMeshConversionOptions DefaultOptions;
 		FUsdMeshConversionOptions();
 
+		// List of the primvars that should be stored on the MeshDescription
+		TSet<FString> PrimvarNames;
+
 		// Additional UE-space transform to bake into the mesh's vertices
 		FTransform AdditionalTransform;
 
@@ -94,6 +110,9 @@ namespace UsdToUnreal
 		// What subdivision level to subdivide to, if the converted mesh is a subdivision mesh.
 		// 0 means no subdivision should take place (default).
 		int32 SubdivisionLevel;
+
+		// How to store the primvars in the MeshDescription
+		EImportPrimvar ImportPrimvars = EImportPrimvar::Standard;
 	};
 
 	/**
