Interchange USD:
* Fixed slow translation when translating USD stage due to SetCustomLocalTransform.
* Fixed translation didn't traverse instance proxies.

#rnx

diff --git a/Engine/Plugins/Interchange/Extensions/OpenUSD/Source/Import/Private/InterchangeUsdTranslator.cpp b/Engine/Plugins/Interchange/Extensions/OpenUSD/Source/Import/Private/InterchangeUsdTranslator.cpp
--- a/Engine/Plugins/Interchange/Extensions/OpenUSD/Source/Import/Private/InterchangeUsdTranslator.cpp
+++ b/Engine/Plugins/Interchange/Extensions/OpenUSD/Source/Import/Private/InterchangeUsdTranslator.cpp
@@ -3201,7 +3201,9 @@ namespace UE::InterchangeUsdTranslator::Private
 			// rest pose on the skeleton asset/ReferenceSkeleton
 			BoneNode->SetCustomBindPoseLocalTransform(&NodeContainer, Bone.LocalRestTransform);
 			BoneNode->SetCustomTimeZeroLocalTransform(&NodeContainer, Bone.LocalRestTransform);
-			BoneNode->SetCustomLocalTransform(&NodeContainer, Bone.LocalRestTransform);
+
+			constexpr bool bResetCache = false;
+			BoneNode->SetCustomLocalTransform(&NodeContainer, Bone.LocalRestTransform, bResetCache);
 
 			NodeContainer.AddNode(BoneNode);
 			NodeContainer.SetNodeParentUid(BoneNodeUid, ParentNode.GetUniqueID());
@@ -3461,7 +3463,9 @@ namespace UE::InterchangeUsdTranslator::Private
 			FTransform Transform = FTransform::Identity;
 			bool bResetTransformStack = false;
 			UsdToUnreal::ConvertXformable(Prim.GetStage(), UE::FUsdTyped(Prim), Transform, UsdUtils::GetEarliestTimeCode(), &bResetTransformStack);
-			SceneNode->SetCustomLocalTransform(&NodeContainer, Transform);
+
+			constexpr bool bResetCache = false;
+			SceneNode->SetCustomLocalTransform(&NodeContainer, Transform, bResetCache);
 
 			// Hide our scene node if it is meant to be hidden in USD
 			//
@@ -3531,7 +3535,8 @@ namespace UE::InterchangeUsdTranslator::Private
 		{
 			Info.ParentNode = SceneNode;
 
-			for (const FUsdPrim& ChildPrim : Prim.GetChildren())
+			constexpr bool bTraverseInstanceProxies = true;
+			for (const FUsdPrim& ChildPrim : Prim.GetFilteredChildren(bTraverseInstanceProxies))
 			{
 				Traverse(ChildPrim, TranslatorImpl, NodeContainer, TranslatorSettings, Info);
 			}
