Interchange USD: Fixed import of geometry caches with typeless prims in their hierarchy. Previously, the typeless prim would prevent the asset from being imported.

diff --git a/Engine/Plugins/Interchange/Extensions/OpenUSD/Source/Import/Private/InterchangeUsdTranslator.cpp b/Engine/Plugins/Interchange/Extensions/OpenUSD/Source/Import/Private/InterchangeUsdTranslator.cpp
--- a/Engine/Plugins/Interchange/Extensions/OpenUSD/Source/Import/Private/InterchangeUsdTranslator.cpp
+++ b/Engine/Plugins/Interchange/Extensions/OpenUSD/Source/Import/Private/InterchangeUsdTranslator.cpp
@@ -3421,12 +3421,16 @@ namespace UE::InterchangeUsdTranslator::Private
 
 		// Only prims that require rendering (and have a renderable parent) get a scene node.
 		// This includes Xforms but also Scopes, which are not Xformable.
+		// Also allow typeless prims to get a scene node otherwise some assets like geometry cache
+		// would not get processed and they need to bake the transforms into the meshes.
 		//
 		// We add a scene node for the pseudoroot in order to make tree traversal easier on the pipeline,
 		// but the pipeline will strip the pseudoroot node as its final step (if desired)
 		const bool bIsImageable = Prim.IsA(TEXT("Imageable"));
-		const bool bNeedsSceneNode = Prim.IsPseudoRoot()														  //
-									 || (bIsImageable && (Info.ParentNode || Prim.GetParent().IsPseudoRoot()))	  //
+		const bool bIsTypeless = Prim.GetTypeName().IsNone();
+		const bool bNeedsSceneNode = Prim.IsPseudoRoot()
+									 || bIsTypeless 
+									 || (bIsImageable && (Info.ParentNode || Prim.GetParent().IsPseudoRoot()))
 									 || Info.bIsLODContainer;	 // We allow typeless prims (so not imageable) to be LOD containers
 
 		UInterchangeSceneNode* SceneNode = nullptr;
