Interchange:
- Fix a crash that may happen because MaterialX material functions are not loaded in the Game Thread when not called from Interchange

diff --git a/Engine/Plugins/Interchange/Extensions/OpenUSD/Source/Import/Private/InterchangeUsdTranslator.cpp b/Engine/Plugins/Interchange/Extensions/OpenUSD/Source/Import/Private/InterchangeUsdTranslator.cpp
--- a/Engine/Plugins/Interchange/Extensions/OpenUSD/Source/Import/Private/InterchangeUsdTranslator.cpp
+++ b/Engine/Plugins/Interchange/Extensions/OpenUSD/Source/Import/Private/InterchangeUsdTranslator.cpp
@@ -4292,6 +4292,8 @@ TArray<FString> UInterchangeUSDTranslator::GetSupportedFormats() const
 	TArray<FString> Extensions;
 	if (GInterchangeEnableUSDImport)
 	{
+		// ensure that MaterialX material functions are loaded in the Game Thread
+		UE::Interchange::MaterialX::AreMaterialFunctionPackagesLoaded();
 		FModuleManager::Get().LoadModuleChecked(TEXT("UnrealUSDWrapper"));
 		UnrealUSDWrapper::AddUsdImportFileFormatDescriptions(Extensions);
 	}
diff --git a/Engine/Plugins/Interchange/Runtime/Source/Import/Public/MaterialX/MaterialXUtils/MaterialXManager.h b/Engine/Plugins/Interchange/Runtime/Source/Import/Public/MaterialX/MaterialXUtils/MaterialXManager.h
--- a/Engine/Plugins/Interchange/Runtime/Source/Import/Public/MaterialX/MaterialXUtils/MaterialXManager.h
+++ b/Engine/Plugins/Interchange/Runtime/Source/Import/Public/MaterialX/MaterialXUtils/MaterialXManager.h
@@ -131,5 +131,5 @@ private:
 namespace UE::Interchange::MaterialX
 {
 	// Load necessary material functions, this function can only be called in the Game Thread
-	bool AreMaterialFunctionPackagesLoaded();
+	INTERCHANGEIMPORT_API bool AreMaterialFunctionPackagesLoaded();
 }
\ No newline at end of file
