Interchange: Remove the "Support" suffixes from the OpenUSD and OpenVDB FriendlyNames

diff --git a/Engine/Plugins/Interchange/Extensions/OpenUSD/InterchangeOpenUSD.uplugin b/Engine/Plugins/Interchange/Extensions/OpenUSD/InterchangeOpenUSD.uplugin
--- a/Engine/Plugins/Interchange/Extensions/OpenUSD/InterchangeOpenUSD.uplugin
+++ b/Engine/Plugins/Interchange/Extensions/OpenUSD/InterchangeOpenUSD.uplugin
@@ -2,7 +2,7 @@
 	"FileVersion": 3,
 	"Version": 1,
 	"VersionName": "1.0",
-	"FriendlyName": "Interchange OpenUSD Support",
+	"FriendlyName": "Interchange OpenUSD",
 	"Description": "Allows translation of OpenUSD files via the Interchange framework",
 	"Category": "Importers",
 	"CreatedBy": "Epic Games, Inc.",
diff --git a/Engine/Plugins/Interchange/Extensions/OpenUSD/Source/Import/Private/InterchangeUsdTranslator.cpp b/Engine/Plugins/Interchange/Extensions/OpenUSD/Source/Import/Private/InterchangeUsdTranslator.cpp
--- a/Engine/Plugins/Interchange/Extensions/OpenUSD/Source/Import/Private/InterchangeUsdTranslator.cpp
+++ b/Engine/Plugins/Interchange/Extensions/OpenUSD/Source/Import/Private/InterchangeUsdTranslator.cpp
@@ -2546,7 +2546,7 @@ namespace UE::InterchangeUsdTranslator::Private
 							Error,
 							TEXT("Failed to find a compatible translator for file '%s'.%s"),
 							*FilePath,
-							bIsOpenVDB ? TEXT(" Is the 'Interchange OpenVDB Support' plugin enabled?") : TEXT("")
+							bIsOpenVDB ? TEXT(" Is the 'Interchange OpenVDB' plugin enabled?") : TEXT("")
 						);
 						continue;
 					}
diff --git a/Engine/Plugins/Interchange/Extensions/OpenVDB/InterchangeOpenVDB.uplugin b/Engine/Plugins/Interchange/Extensions/OpenVDB/InterchangeOpenVDB.uplugin
--- a/Engine/Plugins/Interchange/Extensions/OpenVDB/InterchangeOpenVDB.uplugin
+++ b/Engine/Plugins/Interchange/Extensions/OpenVDB/InterchangeOpenVDB.uplugin
@@ -2,7 +2,7 @@
 	"FileVersion": 3,
 	"Version": 1,
 	"VersionName": "1.0",
-	"FriendlyName": "Interchange OpenVDB Support",
+	"FriendlyName": "Interchange OpenVDB",
 	"Description": "Allows translation of OpenVDB files via the Interchange framework",
 	"Category": "Importers",
 	"CreatedBy": "Epic Games, Inc.",
