USD: Don't register Interchange dialog overrides for the "None" asset type

diff --git a/Engine/Plugins/Interchange/Extensions/OpenUSD/Source/Import/Private/InterchangeOpenUSDImportModule.cpp b/Engine/Plugins/Interchange/Extensions/OpenUSD/Source/Import/Private/InterchangeOpenUSDImportModule.cpp
--- a/Engine/Plugins/Interchange/Extensions/OpenUSD/Source/Import/Private/InterchangeOpenUSDImportModule.cpp
+++ b/Engine/Plugins/Interchange/Extensions/OpenUSD/Source/Import/Private/InterchangeOpenUSDImportModule.cpp
@@ -67,7 +67,6 @@ void FInterchangeOpenUSDImportModule::StartupModule()
 		{
 			// TODO: Add a Count value to the enum and then use ENUM_RANGE_BY_COUNT?
 			static const TArray<EInterchangeTranslatorAssetType> AssetTypes = {
-				EInterchangeTranslatorAssetType::None,
 				EInterchangeTranslatorAssetType::Textures,
 				EInterchangeTranslatorAssetType::Materials,
 				EInterchangeTranslatorAssetType::Meshes,
