USD: Add unsupported warning when parsing the SpatialAudio, PointInstancer, and BasisCurves schemas via USD Interchange.

diff --git a/Engine/Plugins/Interchange/Extensions/OpenUSD/Source/Import/Private/InterchangeUsdTranslator.cpp b/Engine/Plugins/Interchange/Extensions/OpenUSD/Source/Import/Private/InterchangeUsdTranslator.cpp
--- a/Engine/Plugins/Interchange/Extensions/OpenUSD/Source/Import/Private/InterchangeUsdTranslator.cpp
+++ b/Engine/Plugins/Interchange/Extensions/OpenUSD/Source/Import/Private/InterchangeUsdTranslator.cpp
@@ -346,6 +346,9 @@ namespace UE::InterchangeUsdTranslator::Private
 		TSharedPtr<FUsdSchemaTranslationContext> TranslationContext;
 		FUsdInfoCache* InfoCache = nullptr;
 
+		// Owned by the translator itself
+		TObjectPtr<UInterchangeResultsContainer> ResultsContainer;
+
 		// We store temp stages in here that we open in order to parse stuff inside of inactive variants of our main UsdStage.
 		// We do this because the payload data are retrieved concurrently, and toggling variants mutates the current stage.
 		mutable FRWLock PrimPathToVariantToStageLock;
@@ -3379,8 +3382,9 @@ namespace UE::InterchangeUsdTranslator::Private
 	{
 		TRACE_CPUPROFILER_EVENT_SCOPE(Traverse)
 
-		FString SceneNodeUid = Prim.GetPrimPath().GetString();
-		FString DisplayLabel(Prim.GetName().ToString());
+		const FString SceneNodeUid = Prim.GetPrimPath().GetString();
+		const FString DisplayLabel(Prim.GetName().ToString());
+		const FName TypeName = Prim.GetTypeName();
 
 		// Do this before generating other nodes as they may need the updated info
 		Info.UpdateWithCurrentPrim(Prim);
@@ -3418,6 +3422,21 @@ namespace UE::InterchangeUsdTranslator::Private
 		{
 			AssetNodeUids = TranslatorImpl.AddVolumeNodes(Prim, NodeContainer, VolumeMaterialInstanceUid, bNeedsVolumeTrack);
 		}
+		else
+		{
+			const static TSet<FName> KnownUnsupported = {TEXT("SpatialAudio"), TEXT("PointInstancer"), TEXT("BasisCurves")};
+
+			if (KnownUnsupported.Contains(TypeName))
+			{
+				UInterchangeResultWarning_Generic* Message = TranslatorImpl.ResultsContainer->Add<UInterchangeResultWarning_Generic>();
+				Message->SourceAssetName = TranslatorImpl.UsdStage.GetRootLayer().GetRealPath();
+				Message->Text = FText::Format(
+					LOCTEXT("UnsupportedSchemaType", "Prim '{0}' has schema '{1}', which is not yet supported via USD Interchange."),
+					FText::FromString(SceneNodeUid),
+					FText::FromName(TypeName)
+				);
+			}
+		}
 
 		// Only prims that require rendering (and have a renderable parent) get a scene node.
 		// This includes Xforms but also Scopes, which are not Xformable.
@@ -3427,7 +3446,7 @@ namespace UE::InterchangeUsdTranslator::Private
 		// We add a scene node for the pseudoroot in order to make tree traversal easier on the pipeline,
 		// but the pipeline will strip the pseudoroot node as its final step (if desired)
 		const bool bIsImageable = Prim.IsA(TEXT("Imageable"));
-		const bool bIsTypeless = Prim.GetTypeName().IsNone();
+		const bool bIsTypeless = TypeName.IsNone();
 		const bool bNeedsSceneNode = Prim.IsPseudoRoot()
 									 || bIsTypeless 
 									 || (bIsImageable && (Info.ParentNode || Prim.GetParent().IsPseudoRoot()))
@@ -4317,6 +4336,7 @@ bool UInterchangeUSDTranslator::Translate(UInterchangeBaseNodeContainer& NodeCon
 	{
 		return false;
 	}
+	ImplPtr->ResultsContainer = Results;
 
 	UInterchangeUsdTranslatorSettings* Settings = Cast<UInterchangeUsdTranslatorSettings>(GetSettings());
 	if (!Settings)
