Interchange USD:
- Refactored primvars baking
- Use of FPatchBasedMeshUVGenerator::AutoComputeUVs instead of FStaticMeshOperations::AddUVChannel as it gives more consistent results on the UVs, and it's what it's used internally by the ModelingTool

diff --git a/Engine/Plugins/Interchange/Extensions/OpenUSD/Source/Import/Private/InterchangeUSDPipeline.cpp b/Engine/Plugins/Interchange/Extensions/OpenUSD/Source/Import/Private/InterchangeUSDPipeline.cpp
--- a/Engine/Plugins/Interchange/Extensions/OpenUSD/Source/Import/Private/InterchangeUSDPipeline.cpp
+++ b/Engine/Plugins/Interchange/Extensions/OpenUSD/Source/Import/Private/InterchangeUSDPipeline.cpp
@@ -3,7 +3,6 @@
 #include "InterchangeUSDPipeline.h"
 
 #include "InterchangeHeterogeneousVolumeActorFactoryNode.h"
-#include "InterchangeManager.h"
 #include "InterchangeMaterialDefinitions.h"
 #include "InterchangeMaterialFactoryNode.h"
 #include "InterchangeMaterialInstanceNode.h"
@@ -16,27 +15,22 @@
 #include "InterchangeStaticMeshLodDataNode.h"
 #include "InterchangeVolumeNode.h"
 #include "Nodes/InterchangeBaseNodeContainer.h"
-#include "Nodes/InterchangeUserDefinedAttribute.h"
 
 #include "DynamicMesh/DynamicMesh3.h"
+#include "DynamicMeshToMeshDescription.h"
 #include "Engine/StaticMesh.h"
 #include "Materials/Material.h"
 #include "Materials/MaterialExpressionTextureBase.h"
 #include "Materials/MaterialExpressionTextureSample.h"
-#include "Materials/MaterialInstanceConstant.h"
-#include "Materials/MaterialInstanceDynamic.h"
 #include "MeshDescriptionToDynamicMesh.h"
+#include "Parameterization/PatchBasedMeshUVGenerator.h"
 #include "Sampling/MeshMapBaker.h"
 #include "Sampling/MeshPropertyMapEvaluator.h"
 #include "StaticMeshAttributes.h"
-#include "StaticMeshOperations.h"
 
 #include "UnrealUSDWrapper.h"
 #include "Usd/InterchangeUsdDefinitions.h"
 #include "USDConversionUtils.h"
-#include "USDLog.h"
-#include "USDMaterialUtils.h"
-#include "USDProjectSettings.h"
 #include "UsdWrappers/SdfPath.h"
 
 #if WITH_EDITORONLY_DATA
@@ -1012,6 +1006,42 @@ namespace UE::InterchangeUsdPipeline::Private
 		{
 		}
 
+		template<typename T>
+		struct TVertexInstanceAttributesRefGuard
+		{
+			TVertexInstanceAttributesRefGuard(TVertexInstanceAttributesRef<T>& InAttributesRef)
+				: AttributesRef(InAttributesRef)
+			{
+				// store the original values
+				StoredAttributes.Reserve(InAttributesRef.GetNumElements());
+				for (int32 Index = 0; Index < InAttributesRef.GetNumElements(); ++Index)
+				{
+					StoredAttributes.Emplace(Index, AttributesRef.Get(Index));
+				}
+			}
+
+			~TVertexInstanceAttributesRefGuard()
+			{
+				Reset();
+			}
+
+			/** Reset a TVertexInstanceAttributesRef to its previous state with its previous values */
+			void Reset()
+			{
+				for (const TPair<FVertexInstanceID, T> Pair : StoredAttributes)
+				{
+					AttributesRef[Pair.Key] = Pair.Value;
+				}
+			}
+
+		private:
+			TVertexInstanceAttributesRef<T>& AttributesRef;
+			TMap<FVertexInstanceID, T> StoredAttributes;
+		};
+
+		using FVertexInstanceAttributesNormalsRefGuard = TVertexInstanceAttributesRefGuard<FVector3f>;
+		using FVertexInstanceAttributesVertexColorsRefGuard = TVertexInstanceAttributesRefGuard<FVector4f>;
+
 		void Bake()
 		{
 			using namespace UE::Interchange;
@@ -1022,10 +1052,24 @@ namespace UE::InterchangeUsdPipeline::Private
 			{
 				FMeshDescription MeshDescriptionPrimvar = CreateMeshDescription();
 
+				// DynamicMesh will hold the new UV set, at the end of the process we have to convert back to the MeshDescription, so it also has the new UV set
+				FDynamicMesh3 DynamicMesh;
+
 				FStaticMeshAttributes MeshAttributes(MeshDescriptionPrimvar);
 				TVertexInstanceAttributesRef<FVector4f> VertexColors = MeshAttributes.GetVertexInstanceColors();
 				TVertexInstanceAttributesRef<FVector3f> Normals = MeshAttributes.GetVertexInstanceNormals();
 
+				// we have to safe guard the vertex colors and the normals because we're directly working on them when baking them
+				// so we can restore them at the end
+				// Because the Conversion goes like this:
+				// 1. MeshDescriptionPrimvar alters its VertexColor/Normal attribute with the stored primvar
+				// 2. Conversion to DynamicMesh
+				// 3. Bake
+				// 4. Converting back DynamicMesh to Original MeshDescription with now the newly added UV set
+				// Note: the different conversion from/to DynamicMesh/MeshDescription doesn't preserve the primvars though since they are custom vertex instance attributes
+				FVertexInstanceAttributesVertexColorsRefGuard VertexColorsGuard{ VertexColors };
+				FVertexInstanceAttributesNormalsRefGuard NormalsGuard{ Normals };
+
 				for (int32 Index = 0; Index < NumberOfPrimvar; ++Index)
 				{
 					UMaterialExpressionTextureSample* TextureSampleExpression = GetMaterialExpressionTextureSample(Index);
@@ -1041,7 +1085,7 @@ namespace UE::InterchangeUsdPipeline::Private
 					FactoryNode->GetBooleanAttribute(USD::Primvar::TangentSpace + FString::FromInt(Index), bTangentSpace);
 
 					MeshDescriptionPrimvar.VertexInstanceAttributes().ForEach(
-					[this, &PrimvarName, &VertexColors, &Normals, TextureSampleExpression, &MeshDescriptionPrimvar, bTangentSpace]
+					[this, &PrimvarName, &VertexColors, &Normals, TextureSampleExpression, &MeshDescriptionPrimvar, &DynamicMesh, bTangentSpace]
 					(const FName AttributeName, auto AttributesConstRef)
 					{
 						if (AttributeName != PrimvarName)
@@ -1082,9 +1126,13 @@ namespace UE::InterchangeUsdPipeline::Private
 							}
 						}
 
-						FDynamicMesh3 DynamicMesh;
 						FMeshDescriptionToDynamicMesh Convert;
 						Convert.Convert(&MeshDescriptionPrimvar, DynamicMesh);
+						DynamicMesh.Attributes()->SetNumUVLayers(UVChannel + 1);
+						FDynamicMeshUVOverlay* UVOVerlay = DynamicMesh.Attributes()->GetUVLayer(UVChannel);
+
+						FPatchBasedMeshUVGenerator Generator;
+						Generator.AutoComputeUVs(DynamicMesh, *UVOVerlay);
 
 						const FDynamicMeshAABBTree3 DetailSpatial(&DynamicMesh);
 						FMeshBakerDynamicMeshSampler DetailSampler(&DynamicMesh, &DetailSpatial);
@@ -1122,6 +1170,14 @@ namespace UE::InterchangeUsdPipeline::Private
 					});
 				}
 
+				VertexColorsGuard.Reset();
+				NormalsGuard.Reset();
+
+				FDynamicMeshToMeshDescription Converter;
+				// convert back to the mesh description so it's also adding the new UV set
+				// though this conversion unfortunately doesn't preserve the primvars, do we need them after all that?
+				Converter.Convert(&DynamicMesh, *StaticMeshAsset->GetMeshDescription(0));
+
 				UpdateMaterials();
 			}
 		}
@@ -1160,17 +1216,7 @@ namespace UE::InterchangeUsdPipeline::Private
 
 			if (AreUVsOverlapping(VertexInstanceUVs))
 			{
-				FStaticMeshOperations::FGenerateUVOptions Options;
-				Options.UVMethod = FStaticMeshOperations::EGenerateUVMethod::Default;
-				TArray<FVector2D> UVs;
-				FStaticMeshOperations::GenerateUV(*MeshDescription, Options, UVs);
-				FStaticMeshOperations::AddUVChannel(*MeshDescription);
-				UVChannel = VertexInstanceUVs.GetNumChannels() - 1;
-
-				for (const int32 VertexInstanceID : MeshDescription->VertexInstances().GetElementIDs())
-				{
-					VertexInstanceUVs.Set(VertexInstanceID, UVChannel, FVector2f{UVs[VertexInstanceID]});
-				}
+				UVChannel++;
 			}
 
 			FMeshDescription MeshDescriptionPrimvar(*MeshDescription);
