Interchange USD:
- Use now the same code as Interchange MaterialX translator and validate a document the same way
- Add boolean support in the usd/mtlx shader graph

diff --git a/Engine/Plugins/Interchange/Extensions/OpenUSD/Source/Import/Private/InterchangeUsdTranslator.cpp b/Engine/Plugins/Interchange/Extensions/OpenUSD/Source/Import/Private/InterchangeUsdTranslator.cpp
--- a/Engine/Plugins/Interchange/Extensions/OpenUSD/Source/Import/Private/InterchangeUsdTranslator.cpp
+++ b/Engine/Plugins/Interchange/Extensions/OpenUSD/Source/Import/Private/InterchangeUsdTranslator.cpp
@@ -1095,27 +1095,7 @@ namespace UE::InterchangeUsdTranslator::Private
 			GeomPropValues = ShaderGraph.GetGeomPropValueNames();
 			if (MaterialX::DocumentPtr Document = ShaderGraph.GetDocument())
 			{
-				if (MaterialX::NodePtr SurfaceMaterial = Document->getNode(TCHAR_TO_UTF8(*Prim.GetName().ToString())))
-				{
-					try
-					{
-						TSharedPtr<FMaterialXBase> ShaderTranslator = FMaterialXManager::GetInstance().GetShaderTranslator(
-							SurfaceMaterial->getCategory().c_str(),
-							NodeContainer
-						);
-						bHasShaderGraph = ShaderTranslator.IsValid();
-						if (bHasShaderGraph)
-						{
-							ShaderTranslator->Translate(SurfaceMaterial);
-						}
-					}
-					catch (std::exception& Exception)
-					{
-						UE_LOG(LogUsd, Error, TEXT("%s"), UTF8_TO_TCHAR(Exception.what()));
-					}
-				}
-
-				return bHasShaderGraph;
+				return FMaterialXManager::GetInstance().Translate(Document, NodeContainer);
 			}
 #endif	  // PLATFORM_WINDOWS
 		}
diff --git a/Engine/Plugins/Runtime/USDCore/Source/USDUtilities/Private/USDMaterialXShaderGraph.cpp b/Engine/Plugins/Runtime/USDCore/Source/USDUtilities/Private/USDMaterialXShaderGraph.cpp
--- a/Engine/Plugins/Runtime/USDCore/Source/USDUtilities/Private/USDMaterialXShaderGraph.cpp
+++ b/Engine/Plugins/Runtime/USDCore/Source/USDUtilities/Private/USDMaterialXShaderGraph.cpp
@@ -51,6 +51,7 @@ FUsdMaterialXShaderGraph::FUsdMaterialXShaderGraph(const pxr::UsdPrim& UsdShadeM
 		{"vector3f", "vector3"},
 		{"vector4f", "vector4"},
 		{"int",      "integer"},
+		{"bool",     "boolean"},
 	}
 	, TangentSpaceInputs{
 	    //UsdPreviewSurface
@@ -305,6 +306,10 @@ void FUsdMaterialXShaderGraph::SetInputValue(mx::InputPtr Input, const pxr::UsdS
 	{
 		SetInputValue(int32{});
 	}
+	else if (InputType == "boolean")
+	{
+		SetInputValue(bool{});
+	}
 	else if (InputType == "vector2")
 	{
 		SetInputValue(mx::Vector2{});
