[Interchange-USD] Performance optimizations.

#rnx

diff --git a/Engine/Plugins/Interchange/Editor/Source/Pipelines/Private/SInterchangePipelineConfigurationDialog.cpp b/Engine/Plugins/Interchange/Editor/Source/Pipelines/Private/SInterchangePipelineConfigurationDialog.cpp
--- a/Engine/Plugins/Interchange/Editor/Source/Pipelines/Private/SInterchangePipelineConfigurationDialog.cpp
+++ b/Engine/Plugins/Interchange/Editor/Source/Pipelines/Private/SInterchangePipelineConfigurationDialog.cpp
@@ -1686,6 +1686,7 @@ void SInterchangePipelineConfigurationDialog::UpdatePreviewContainer(bool bUpdat
 		PreviewNodeContainer = nullptr;
 	}
 	PreviewNodeContainer = DuplicateObject<UInterchangeBaseNodeContainer>(BaseNodeContainer.Get(), GetTransientPackage());
+	PreviewNodeContainer->SetChildrenCache(BaseNodeContainer->GetChildrenCache());
 
 	TArray<UInterchangeSourceData*> SourceDatas;
 	SourceDatas.Add(SourceData.Get());
diff --git a/Engine/Plugins/Interchange/Extensions/OpenUSD/Source/Import/Private/InterchangeUsdTranslator.cpp b/Engine/Plugins/Interchange/Extensions/OpenUSD/Source/Import/Private/InterchangeUsdTranslator.cpp
--- a/Engine/Plugins/Interchange/Extensions/OpenUSD/Source/Import/Private/InterchangeUsdTranslator.cpp
+++ b/Engine/Plugins/Interchange/Extensions/OpenUSD/Source/Import/Private/InterchangeUsdTranslator.cpp
@@ -3134,8 +3134,11 @@ namespace UE::InterchangeUsdTranslator::Private
 		//    root bone anyway: The SkeletonPrimNode itself;
 		//  - If a skeleton has no bones at all somehow, we'll still make one "bone" for it (this node).
 		SkeletonPrimNode.AddSpecializedType(UE::Interchange::FSceneNodeStaticData::GetJointSpecializeTypeString());
-		SkeletonPrimNode.SetCustomBindPoseLocalTransform(&NodeContainer, FTransform::Identity);
-		SkeletonPrimNode.SetCustomTimeZeroLocalTransform(&NodeContainer, FTransform::Identity);
+
+		constexpr bool bResetCache = false;
+		SkeletonPrimNode.SetCustomBindPoseLocalTransform(&NodeContainer, FTransform::Identity, bResetCache);
+		SkeletonPrimNode.SetCustomTimeZeroLocalTransform(&NodeContainer, FTransform::Identity, bResetCache);
+
 		const FString SkeletonPrimNodeUid = SkeletonPrimNode.GetUniqueID();
 
 #if WITH_EDITOR
@@ -3178,19 +3181,18 @@ namespace UE::InterchangeUsdTranslator::Private
 			BoneNode->InitializeNode(BoneNodeUid, Bone.Name, EInterchangeNodeContainerType::TranslatedScene);
 			BoneNode->AddSpecializedType(UE::Interchange::FSceneNodeStaticData::GetJointSpecializeTypeString());
 
+			NodeContainer.AddNode(BoneNode);
+			NodeContainer.SetNodeParentUid(BoneNodeUid, ParentNode.GetUniqueID());
+
 			// Note that we use our rest transforms for the Interchange bind pose as well: This because Interchange
 			// will put this on the RefSkeleton and so it will make its way to the Skeleton asset. We already kind
 			// of bake in our skeleton bind pose directly into our skinned mesh, so we really just want to put the
 			// rest pose on the skeleton asset/ReferenceSkeleton
-			BoneNode->SetCustomBindPoseLocalTransform(&NodeContainer, Bone.LocalRestTransform);
-			BoneNode->SetCustomTimeZeroLocalTransform(&NodeContainer, Bone.LocalRestTransform);
-
 			constexpr bool bResetCache = false;
+			BoneNode->SetCustomBindPoseLocalTransform(&NodeContainer, Bone.LocalRestTransform, bResetCache);
+			BoneNode->SetCustomTimeZeroLocalTransform(&NodeContainer, Bone.LocalRestTransform, bResetCache);
 			BoneNode->SetCustomLocalTransform(&NodeContainer, Bone.LocalRestTransform, bResetCache);
 
-			NodeContainer.AddNode(BoneNode);
-			NodeContainer.SetNodeParentUid(BoneNodeUid, ParentNode.GetUniqueID());
-
 			BoneToUidAndBoneIndex.Add(ConcatBonePath, {BoneNodeUid, BoneIndex});
 
 			for (int32 ChildIndex : Bone.ChildIndices)
@@ -3437,7 +3439,13 @@ namespace UE::InterchangeUsdTranslator::Private
 		{
 			SceneNode = NewObject<UInterchangeSceneNode>(&NodeContainer);
 			SceneNode->InitializeNode(SceneNodeUid, DisplayLabel, EInterchangeNodeContainerType::TranslatedScene);
+
 			NodeContainer.AddNode(SceneNode);
+			// Connect parent and child scene nodes
+			if (Info.ParentNode)
+			{
+				NodeContainer.SetNodeParentUid(SceneNodeUid, Info.ParentNode->GetUniqueID());
+			}
 
 			// Store our purpose here to be filtered on the pipeline
 			const EUsdPurpose PrimPurpose = IUsdPrim::GetPurpose(Prim);
@@ -3517,12 +3525,6 @@ namespace UE::InterchangeUsdTranslator::Private
 				SceneNode->SetCustomAssetInstanceUid(AssetNodeUids[0]);
 			}
 
-			// Connect parent and child scene nodes
-			if (Info.ParentNode)
-			{
-				NodeContainer.SetNodeParentUid(SceneNode->GetUniqueID(), Info.ParentNode->GetUniqueID());
-			}
-
 			// Add animation tracks
 			if (bIsImageable)
 			{
diff --git a/Engine/Plugins/Interchange/Runtime/Source/Import/Private/Mesh/InterchangeSkeletonHelper.cpp b/Engine/Plugins/Interchange/Runtime/Source/Import/Private/Mesh/InterchangeSkeletonHelper.cpp
--- a/Engine/Plugins/Interchange/Runtime/Source/Import/Private/Mesh/InterchangeSkeletonHelper.cpp
+++ b/Engine/Plugins/Interchange/Runtime/Source/Import/Private/Mesh/InterchangeSkeletonHelper.cpp
@@ -189,7 +189,12 @@ namespace UE::Interchange::Private
 
 		// Iterate children
 		TArray<FString>* CachedChildren = NodeContainer->GetCachedNodeChildrenUids(JointUid);
-		check(CachedChildren); // Call ComputeChildrenCache() prior to calling this function
+
+		//Note: If CacheChildren is null, that just means that the JointUid does not have children.
+		if (!CachedChildren)
+		{
+			return;
+		}
 
 		const TArray<FString>& ChildrenIds = *CachedChildren;
 		for (int32 ChildIndex = 0; ChildIndex < ChildrenIds.Num(); ++ChildIndex)
diff --git a/Engine/Plugins/Interchange/Runtime/Source/Pipelines/Private/InterchangeGenericAnimationPipeline.cpp b/Engine/Plugins/Interchange/Runtime/Source/Pipelines/Private/InterchangeGenericAnimationPipeline.cpp
--- a/Engine/Plugins/Interchange/Runtime/Source/Pipelines/Private/InterchangeGenericAnimationPipeline.cpp
+++ b/Engine/Plugins/Interchange/Runtime/Source/Pipelines/Private/InterchangeGenericAnimationPipeline.cpp
@@ -813,7 +813,10 @@ void UInterchangeGenericAnimationPipeline::CreateAnimSequenceFactoryNode(UInterc
 		// NOTE: Could this be added as an array of FString attributes on the UInterchangeSkeletalAnimationTrackNode
 #if WITH_EDITOR
 		//Iterate all joints to set the meta data value in the anim sequence factory node
-		BaseNodeContainer->ComputeChildrenCache();
+		
+		//Note: We shouldn't be able to get to this point without the ChildrenCache already set, so there is no need to call ComputeChildrenCache.
+		// (Also trying to avoid recreating the ChildrenCache for every AnimSequenceFactoryNode creation call.)
+
 		UE::Interchange::Private::FSkeletonHelper::RecursiveAddSkeletonMetaDataValues(BaseNodeContainer, AnimSequenceFactoryNode, RootJointUid);
 #endif //WITH_EDITOR
 
diff --git a/Engine/Plugins/Interchange/Runtime/Source/Pipelines/Private/InterchangeGenericCommonSkeletalMeshesAndAnimationsProperties.cpp b/Engine/Plugins/Interchange/Runtime/Source/Pipelines/Private/InterchangeGenericCommonSkeletalMeshesAndAnimationsProperties.cpp
--- a/Engine/Plugins/Interchange/Runtime/Source/Pipelines/Private/InterchangeGenericCommonSkeletalMeshesAndAnimationsProperties.cpp
+++ b/Engine/Plugins/Interchange/Runtime/Source/Pipelines/Private/InterchangeGenericCommonSkeletalMeshesAndAnimationsProperties.cpp
@@ -63,7 +63,10 @@ UInterchangeSkeletonFactoryNode* UInterchangeGenericCommonSkeletalMeshesAndAnima
 	}
 #if WITH_EDITOR
 	//Iterate all joints to set the meta data value in the skeleton node
-	BaseNodeContainer->ComputeChildrenCache();
+	
+	//Note: We shouldn't be able to get to this point without the ChildrenCache already set, so there is no need to call ComputeChildrenCache.
+	// (Also trying to avoid recreating the ChildrenCache for every SkeletonFactoryNode creation call.)
+
 	UE::Interchange::Private::FSkeletonHelper::RecursiveAddSkeletonMetaDataValues(BaseNodeContainer, SkeletonFactoryNode, RootJointUid);
 #endif //WITH_EDITOR
 
diff --git a/Engine/Source/Runtime/Interchange/Core/Private/Nodes/InterchangeBaseNodeContainer.cpp b/Engine/Source/Runtime/Interchange/Core/Private/Nodes/InterchangeBaseNodeContainer.cpp
--- a/Engine/Source/Runtime/Interchange/Core/Private/Nodes/InterchangeBaseNodeContainer.cpp
+++ b/Engine/Source/Runtime/Interchange/Core/Private/Nodes/InterchangeBaseNodeContainer.cpp
@@ -297,18 +297,7 @@ TArray<FString> UInterchangeBaseNodeContainer::GetNodeChildrenUids(const FString
 		return *CacheChildrenPtr;
 	}
 
-	// Update the cache
-	TArray<FString>& CacheChildren = ChildrenCache.Add(NodeUniqueID);
-	for (const auto& NodeKeyValue : Nodes)
-	{
-		if (NodeKeyValue.Value->GetParentUid() == NodeUniqueID)
-		{
-			CacheChildren.Add(NodeKeyValue.Key);
-		}
-	}
-	//Reorder the children cache
-	InternalReorderChildren(CacheChildren);
-	return CacheChildren;
+	return {};
 }
 
 TArray<FString>* UInterchangeBaseNodeContainer::GetCachedNodeChildrenUids(const FString& NodeUniqueID) const
@@ -431,9 +420,6 @@ void UInterchangeBaseNodeContainer::ComputeChildrenCache()
 	ResetChildrenCache();
 	for (const auto& NodeKeyValue : Nodes)
 	{
-		//Force all node to have a cache, even if there is no parent
-		ChildrenCache.FindOrAdd(NodeKeyValue.Key);
-
 		//Update the parent cache
 		const FString ParentUid = NodeKeyValue.Value->GetParentUid();
 		if (!ParentUid.IsEmpty())
@@ -448,6 +434,19 @@ void UInterchangeBaseNodeContainer::ComputeChildrenCache()
 	}
 }
 
+void UInterchangeBaseNodeContainer::SetChildrenCache(const TMap<FString, TArray<FString>>& InChildrenCache)
+{
+	TRACE_CPUPROFILER_EVENT_SCOPE(UInterchangeBaseNodeContainer::SetChildrenCache)
+
+	ChildrenCache.Reset();
+	ChildrenCache.Append(InChildrenCache);
+}
+
+TMap<FString, TArray<FString>>& UInterchangeBaseNodeContainer::GetChildrenCache()
+{
+	return ChildrenCache;
+}
+
 void UInterchangeBaseNodeContainer::InternalReorderChildren(TArray<FString>& Children) const
 {
 	TRACE_CPUPROFILER_EVENT_SCOPE(UInterchangeBaseNodeContainer::InternalReorderChildren)
diff --git a/Engine/Source/Runtime/Interchange/Core/Public/Nodes/InterchangeBaseNodeContainer.h b/Engine/Source/Runtime/Interchange/Core/Public/Nodes/InterchangeBaseNodeContainer.h
--- a/Engine/Source/Runtime/Interchange/Core/Public/Nodes/InterchangeBaseNodeContainer.h
+++ b/Engine/Source/Runtime/Interchange/Core/Public/Nodes/InterchangeBaseNodeContainer.h
@@ -223,6 +223,12 @@ public:
 	UFUNCTION(BlueprintCallable, Category = "Interchange | Node Container")
 	INTERCHANGECORE_API void ComputeChildrenCache();
 
+	/* Sets the children cache from an incoming data set. */
+	INTERCHANGECORE_API void SetChildrenCache(const TMap<FString, TArray<FString>>& InChildrenCache);
+
+	/* Gets the children cache. */
+	INTERCHANGECORE_API TMap<FString, TArray<FString>>& GetChildrenCache();
+
 	/* Reset the cache of children UIDs. */
 	UFUNCTION(BlueprintCallable, Category = "Interchange | Node Container")
 	void ResetChildrenCache()
diff --git a/Engine/Source/Runtime/Interchange/Engine/Private/Tasks/InterchangeTaskTranslator.cpp b/Engine/Source/Runtime/Interchange/Engine/Private/Tasks/InterchangeTaskTranslator.cpp
--- a/Engine/Source/Runtime/Interchange/Engine/Private/Tasks/InterchangeTaskTranslator.cpp
+++ b/Engine/Source/Runtime/Interchange/Engine/Private/Tasks/InterchangeTaskTranslator.cpp
@@ -54,6 +54,7 @@ void UE::Interchange::FTaskTranslator::Execute()
 	//Translate the source data
 	UInterchangeBaseNodeContainer& BaseNodeContainer = *(AsyncHelper->BaseNodeContainers[SourceIndex].Get());
 	Translator->Translate(BaseNodeContainer);
-	//Make sure the base node container cache is computed for all translated node
-	BaseNodeContainer.ComputeChildrenCache();
+
+	//Note: At this point All childrenCache should already be set, as SetNodeParentUid already handles the ChildrenCache.
+	//		(and FTaskParsing ComputesChildrenCache for any prior existing data/node set.)
 }
