[USD Interchange] Adds import support for (new) LodSubtreeAPI schema.

This schema enables the encoding of LODs from DCC's in a similar fashion to an FBX LOD Group. The name LodSubtreeAPI, rather than LodGroupAPI,
is chosen to avoid possible confusion with the USD kind hierarchy concept, where kind=group is very common. Use of the schema is independent
of any particular USD kind hierarchy.

Unlike the "LOD" variant set approach, the levels do not require any particular naming since they are described by relationship. All levels
are composed on the USD stage at the same time. This should hopefully simplify import/export from both UE and DCC's since there are fewer
restrictions such as specific mesh names etc. For Maya, translation from/to lodGroup nodes is possible without any adjustment to the Maya scene,
giving us parity with FBX LOD import/export.

DCC import/export will not author this schema by default of course. However both Maya and Blender USD have plugin mechanisms that allow user
code to run additional operations following a successful import/export. Both applications refer to these plugins as "chasers". Using these it is
possible to apply and configure LodSubtreeAPI as part of the regular import/export process. Houdini has native USD support via LOPs and so
applying the schema in a LOPs context is straightforward.

Python chaser examples for Maya and Blender along with a python Houdini LOPs node will be included in the documentation.

diff --git a/Engine/Plugins/Interchange/Extensions/OpenUSD/Source/Import/Private/InterchangeUsdTranslator.cpp b/Engine/Plugins/Interchange/Extensions/OpenUSD/Source/Import/Private/InterchangeUsdTranslator.cpp
--- a/Engine/Plugins/Interchange/Extensions/OpenUSD/Source/Import/Private/InterchangeUsdTranslator.cpp
+++ b/Engine/Plugins/Interchange/Extensions/OpenUSD/Source/Import/Private/InterchangeUsdTranslator.cpp
@@ -447,6 +447,70 @@ namespace UE::InterchangeUsdTranslator::Private
 		return {};
 	};
 
+	TArray<UE::FUsdPrim> CheckLodAPIAndGetChildren(const UE::FUsdPrim& Prim, UInterchangeSceneNode* SceneNode)
+	{
+		if (!SceneNode || !Prim.HasAPI(*UsdToUnreal::ConvertToken(UnrealIdentifiers::LodSubtreeAPI)))
+		{
+			constexpr bool bTraverseInstanceProxies = true;
+			return  Prim.GetFilteredChildren(bTraverseInstanceProxies);
+		}
+
+		UE::FUsdRelationship LevelsRel = Prim.GetRelationship(*UsdToUnreal::ConvertToken(UnrealIdentifiers::UnrealLodSubtreeLevels));
+		if (!LevelsRel)
+		{   
+			UE_LOG(
+				LogUsd, Warning, TEXT("LOD subtree '%s' is missing required relationship '%s'"),
+				*Prim.GetPrimPath().GetString(),
+				*UsdToUnreal::ConvertToken(UnrealIdentifiers::UnrealLodSubtreeLevels)
+			);
+			return {};
+		}
+
+		TArray<UE::FSdfPath> Targets;
+		if (!LevelsRel.GetTargets(Targets) || Targets.IsEmpty())
+		{
+			UE_LOG(
+				LogUsd, Warning, TEXT("LOD subtree '%s' has no LODs specified by relationship '%s'"),
+				*Prim.GetPrimPath().GetString(),
+				*UsdToUnreal::ConvertToken(UnrealIdentifiers::UnrealLodSubtreeLevels)
+			);
+			return {};
+		}
+
+		TArray<UE::FUsdPrim> SubtreePrims;
+
+		for (const UE::FSdfPath& TargetPath : Targets)
+		{
+			// Members must be direct children
+			if (TargetPath.GetParentPath() != Prim.GetPrimPath())
+			{
+				UE_LOG(
+					LogUsd, Warning, TEXT("Ignoring LOD '%s' that is not a direct child of '%s'"),
+					*TargetPath.GetString(), *Prim.GetPrimPath().GetString()
+				);
+				continue;
+			}
+				
+			FUsdPrim SubtreePrim = Prim.GetStage().GetPrimAtPath(TargetPath);
+			if (!SubtreePrim)
+			{
+				UE_LOG(
+					LogUsd, Warning, TEXT("Ignoring invalid or missing LOD '%s' specified by '%s'"),
+					*TargetPath.GetString(), *Prim.GetPrimPath().GetString()
+				);
+			}
+			
+			SubtreePrims.Emplace(SubtreePrim);
+		}
+
+		if (!SubtreePrims.IsEmpty())
+		{
+			SceneNode->AddSpecializedType(UE::Interchange::FSceneNodeStaticData::GetLodGroupSpecializeTypeString());
+		}
+
+		return SubtreePrims;
+	}
+
 	bool CheckAndChopPayloadPrefix(FString& PayloadKey, const FString& Prefix)
 	{
 		if (PayloadKey.StartsWith(Prefix))
@@ -3540,8 +3604,7 @@ namespace UE::InterchangeUsdTranslator::Private
 		{
 			Info.ParentNode = SceneNode;
 
-			constexpr bool bTraverseInstanceProxies = true;
-			for (const FUsdPrim& ChildPrim : Prim.GetFilteredChildren(bTraverseInstanceProxies))
+			for (const FUsdPrim& ChildPrim : CheckLodAPIAndGetChildren(Prim, SceneNode))
 			{
 				Traverse(ChildPrim, TranslatorImpl, NodeContainer, TranslatorSettings, Info);
 			}
diff --git a/Engine/Plugins/Runtime/USDCore/Resources/UsdResources/Linux/plugins/unreal/resources/plugInfo.json b/Engine/Plugins/Runtime/USDCore/Resources/UsdResources/Linux/plugins/unreal/resources/plugInfo.json
--- a/Engine/Plugins/Runtime/USDCore/Resources/UsdResources/Linux/plugins/unreal/resources/plugInfo.json
+++ b/Engine/Plugins/Runtime/USDCore/Resources/UsdResources/Linux/plugins/unreal/resources/plugInfo.json
@@ -14,6 +14,7 @@
                         "bases": [
                             "UsdAPISchemaBase"
                         ], 
+                        "schemaIdentifier": "CollapsingAPI", 
                         "schemaKind": "singleApplyAPI"
                     }, 
                     "unrealControlRigAPI": {
@@ -28,6 +29,7 @@
                         "bases": [
                             "UsdAPISchemaBase"
                         ], 
+                        "schemaIdentifier": "ControlRigAPI", 
                         "schemaKind": "singleApplyAPI"
                     }, 
                     "unrealGroomAPI": {
@@ -42,6 +44,7 @@
                         "bases": [
                             "UsdAPISchemaBase"
                         ], 
+                        "schemaIdentifier": "GroomAPI", 
                         "schemaKind": "singleApplyAPI"
                     }, 
                     "unrealGroomBindingAPI": {
@@ -56,6 +59,7 @@
                         "bases": [
                             "UsdAPISchemaBase"
                         ], 
+                        "schemaIdentifier": "GroomBindingAPI", 
                         "schemaKind": "singleApplyAPI"
                     }, 
                     "unrealLiveLinkAPI": {
@@ -66,6 +70,22 @@
                         "bases": [
                             "UsdAPISchemaBase"
                         ], 
+                        "schemaIdentifier": "LiveLinkAPI", 
+                        "schemaKind": "singleApplyAPI"
+                    }, 
+                    "unrealLodSubtreeAPI": {
+                        "alias": {
+                            "UsdSchemaBase": "LodSubtreeAPI"
+                        }, 
+                        "apiSchemaCanOnlyApplyTo": [
+                            "Scope", 
+                            "Xform"
+                        ], 
+                        "autoGenerated": true, 
+                        "bases": [
+                            "UsdAPISchemaBase"
+                        ], 
+                        "schemaIdentifier": "LodSubtreeAPI", 
                         "schemaKind": "singleApplyAPI"
                     }, 
                     "unrealSparseVolumeTextureAPI": {
@@ -80,6 +100,7 @@
                         "bases": [
                             "UsdAPISchemaBase"
                         ], 
+                        "schemaIdentifier": "SparseVolumeTextureAPI", 
                         "schemaKind": "singleApplyAPI"
                     }
                 }
diff --git a/Engine/Plugins/Runtime/USDCore/Resources/UsdResources/Mac/plugins/unreal/resources/plugInfo.json b/Engine/Plugins/Runtime/USDCore/Resources/UsdResources/Mac/plugins/unreal/resources/plugInfo.json
--- a/Engine/Plugins/Runtime/USDCore/Resources/UsdResources/Mac/plugins/unreal/resources/plugInfo.json
+++ b/Engine/Plugins/Runtime/USDCore/Resources/UsdResources/Mac/plugins/unreal/resources/plugInfo.json
@@ -14,6 +14,7 @@
                         "bases": [
                             "UsdAPISchemaBase"
                         ], 
+                        "schemaIdentifier": "CollapsingAPI", 
                         "schemaKind": "singleApplyAPI"
                     }, 
                     "unrealControlRigAPI": {
@@ -28,6 +29,7 @@
                         "bases": [
                             "UsdAPISchemaBase"
                         ], 
+                        "schemaIdentifier": "ControlRigAPI", 
                         "schemaKind": "singleApplyAPI"
                     }, 
                     "unrealGroomAPI": {
@@ -42,6 +44,7 @@
                         "bases": [
                             "UsdAPISchemaBase"
                         ], 
+                        "schemaIdentifier": "GroomAPI", 
                         "schemaKind": "singleApplyAPI"
                     }, 
                     "unrealGroomBindingAPI": {
@@ -56,6 +59,7 @@
                         "bases": [
                             "UsdAPISchemaBase"
                         ], 
+                        "schemaIdentifier": "GroomBindingAPI", 
                         "schemaKind": "singleApplyAPI"
                     }, 
                     "unrealLiveLinkAPI": {
@@ -66,6 +70,22 @@
                         "bases": [
                             "UsdAPISchemaBase"
                         ], 
+                        "schemaIdentifier": "LiveLinkAPI", 
+                        "schemaKind": "singleApplyAPI"
+                    }, 
+                    "unrealLodSubtreeAPI": {
+                        "alias": {
+                            "UsdSchemaBase": "LodSubtreeAPI"
+                        }, 
+                        "apiSchemaCanOnlyApplyTo": [
+                            "Scope", 
+                            "Xform"
+                        ], 
+                        "autoGenerated": true, 
+                        "bases": [
+                            "UsdAPISchemaBase"
+                        ], 
+                        "schemaIdentifier": "LodSubtreeAPI", 
                         "schemaKind": "singleApplyAPI"
                     }, 
                     "unrealSparseVolumeTextureAPI": {
@@ -80,6 +100,7 @@
                         "bases": [
                             "UsdAPISchemaBase"
                         ], 
+                        "schemaIdentifier": "SparseVolumeTextureAPI", 
                         "schemaKind": "singleApplyAPI"
                     }
                 }
diff --git a/Engine/Plugins/Runtime/USDCore/Resources/UsdResources/Win64/plugins/unreal/resources/plugInfo.json b/Engine/Plugins/Runtime/USDCore/Resources/UsdResources/Win64/plugins/unreal/resources/plugInfo.json
--- a/Engine/Plugins/Runtime/USDCore/Resources/UsdResources/Win64/plugins/unreal/resources/plugInfo.json
+++ b/Engine/Plugins/Runtime/USDCore/Resources/UsdResources/Win64/plugins/unreal/resources/plugInfo.json
@@ -14,6 +14,7 @@
                         "bases": [
                             "UsdAPISchemaBase"
                         ], 
+                        "schemaIdentifier": "CollapsingAPI", 
                         "schemaKind": "singleApplyAPI"
                     }, 
                     "unrealControlRigAPI": {
@@ -28,6 +29,7 @@
                         "bases": [
                             "UsdAPISchemaBase"
                         ], 
+                        "schemaIdentifier": "ControlRigAPI", 
                         "schemaKind": "singleApplyAPI"
                     }, 
                     "unrealGroomAPI": {
@@ -42,6 +44,7 @@
                         "bases": [
                             "UsdAPISchemaBase"
                         ], 
+                        "schemaIdentifier": "GroomAPI", 
                         "schemaKind": "singleApplyAPI"
                     }, 
                     "unrealGroomBindingAPI": {
@@ -56,6 +59,7 @@
                         "bases": [
                             "UsdAPISchemaBase"
                         ], 
+                        "schemaIdentifier": "GroomBindingAPI", 
                         "schemaKind": "singleApplyAPI"
                     }, 
                     "unrealLiveLinkAPI": {
@@ -66,6 +70,22 @@
                         "bases": [
                             "UsdAPISchemaBase"
                         ], 
+                        "schemaIdentifier": "LiveLinkAPI", 
+                        "schemaKind": "singleApplyAPI"
+                    }, 
+                    "unrealLodSubtreeAPI": {
+                        "alias": {
+                            "UsdSchemaBase": "LodSubtreeAPI"
+                        }, 
+                        "apiSchemaCanOnlyApplyTo": [
+                            "Scope", 
+                            "Xform"
+                        ], 
+                        "autoGenerated": true, 
+                        "bases": [
+                            "UsdAPISchemaBase"
+                        ], 
+                        "schemaIdentifier": "LodSubtreeAPI", 
                         "schemaKind": "singleApplyAPI"
                     }, 
                     "unrealSparseVolumeTextureAPI": {
@@ -80,6 +100,7 @@
                         "bases": [
                             "UsdAPISchemaBase"
                         ], 
+                        "schemaIdentifier": "SparseVolumeTextureAPI", 
                         "schemaKind": "singleApplyAPI"
                     }
                 }
diff --git a/Engine/Plugins/Runtime/USDCore/Source/UnrealUSDWrapper/Private/UnrealUSDWrapper.cpp b/Engine/Plugins/Runtime/USDCore/Source/UnrealUSDWrapper/Private/UnrealUSDWrapper.cpp
--- a/Engine/Plugins/Runtime/USDCore/Source/UnrealUSDWrapper/Private/UnrealUSDWrapper.cpp
+++ b/Engine/Plugins/Runtime/USDCore/Source/UnrealUSDWrapper/Private/UnrealUSDWrapper.cpp
@@ -88,6 +88,8 @@ namespace UnrealIdentifiers
 	 * named LOD0, LOD1, LOD2, etc
 	 */
 	const TfToken LOD("LOD");
+	const TfToken LodSubtreeAPI("LodSubtreeAPI");
+	const TfToken UnrealLodSubtreeLevels = TfToken("unreal:lodSubtree:levels");
 
 	const TfToken MaterialAssignment = TfToken("unrealMaterial");
 	const TfToken Unreal = TfToken("unreal");
diff --git a/Engine/Plugins/Runtime/USDCore/Source/UnrealUSDWrapper/Public/UnrealUSDWrapper.h b/Engine/Plugins/Runtime/USDCore/Source/UnrealUSDWrapper/Public/UnrealUSDWrapper.h
--- a/Engine/Plugins/Runtime/USDCore/Source/UnrealUSDWrapper/Public/UnrealUSDWrapper.h
+++ b/Engine/Plugins/Runtime/USDCore/Source/UnrealUSDWrapper/Public/UnrealUSDWrapper.h
@@ -413,6 +413,8 @@ namespace UnrealIdentifiers
 {
 #if USE_USD_SDK
 	extern UNREALUSDWRAPPER_API const pxr::TfToken LOD;
+	extern UNREALUSDWRAPPER_API const pxr::TfToken LodSubtreeAPI;
+	extern UNREALUSDWRAPPER_API const pxr::TfToken UnrealLodSubtreeLevels;
 
 	/* Attribute name when assigning Unreal materials to UsdGeomMeshes */
 	extern UNREALUSDWRAPPER_API const pxr::TfToken MaterialAssignment;
