USD: Use Skeleton as a "bound skeleton" fallback when parsing skeletal animations.

diff --git a/Engine/Plugins/Interchange/Extensions/OpenUSD/Source/Import/Private/InterchangeUsdTranslator.cpp b/Engine/Plugins/Interchange/Extensions/OpenUSD/Source/Import/Private/InterchangeUsdTranslator.cpp
--- a/Engine/Plugins/Interchange/Extensions/OpenUSD/Source/Import/Private/InterchangeUsdTranslator.cpp
+++ b/Engine/Plugins/Interchange/Extensions/OpenUSD/Source/Import/Private/InterchangeUsdTranslator.cpp
@@ -3301,6 +3301,20 @@ namespace UE::InterchangeUsdTranslator::Private
 		{
 			Info.SkelJointNames->Add(Bone.Name);
 		}
+
+		// Prefer flagging the Skeleton prim itself as the BoundSkeletonPrimPath at this point in the hierarchy, even
+		// preferring it over any explicit skel:skeleton relationship.
+		// This does not seem technically correct, but is useful in case the Skeleton prim has a skel:animationSource
+		// relationship directly on it, which seems to animate in usdview and is advertised as a supported case.
+		// References:
+		// - [REDACTED_PROFILE_URL]/assets/blob/main/test_assets/USDZ/CesiumMan/CesiumMan.usdz
+		// - https://openusd.org/release/api/_usd_skel__o_m.html
+		// - [REDACTED_PROFILE_URL]/OpenUSD/issues/3532
+		if (UE::FUsdPrim SkeletonPrim = SkelQuery.GetSkeleton())
+		{
+			Info.BoundSkeletonPrimPath = MakeShared<FString>(SkeletonPrim.GetPrimPath().GetString());
+		}
+
 		{
 			FWriteScopeLock ScopedInfoWriteLock{TranslatorImpl.CachedTraversalInfoLock};
 			TranslatorImpl.NodeUidToCachedTraversalInfo.Add(SkeletonPrimNodeUid, Info);
@@ -4042,13 +4056,13 @@ namespace UE::InterchangeUsdTranslator::Private
 		{
 			FReadScopeLock ReadLock{Impl.CachedTraversalInfoLock};
 
-			const FTraversalInfo* MeshInfo = Impl.NodeUidToCachedTraversalInfo.Find(SkeletonPrimPath);
-			if (!MeshInfo)
+			const FTraversalInfo* Info = Impl.NodeUidToCachedTraversalInfo.Find(SkeletonPrimPath);
+			if (!Info)
 			{
 				return false;
 			}
 
-			SkelQuery = MeshInfo->ResolveSkelQuery(Impl.UsdStage);
+			SkelQuery = Info->ResolveSkelQuery(Impl.UsdStage);
 			if (!SkelQuery)
 			{
 				return false;
@@ -4178,13 +4192,13 @@ namespace UE::InterchangeUsdTranslator::Private
 			{
 				FReadScopeLock ReadLock{Impl.CachedTraversalInfoLock};
 
-				const FTraversalInfo* MeshInfo = Impl.NodeUidToCachedTraversalInfo.Find(SkeletonPrimPath);
-				if (!MeshInfo)
+				const FTraversalInfo* Info = Impl.NodeUidToCachedTraversalInfo.Find(SkeletonPrimPath);
+				if (!Info)
 				{
 					return false;
 				}
 
-				SkelQuery = MeshInfo->ResolveSkelQuery(Impl.UsdStage);
+				SkelQuery = Info->ResolveSkelQuery(Impl.UsdStage);
 				if (!SkelQuery)
 				{
 					return false;
