USD: Assign AnimSequences to spawned SkeletalMeshComponents when importing USD stages through Interchange.

diff --git a/Engine/Plugins/Interchange/Extensions/OpenUSD/Source/Import/Private/InterchangeUsdTranslator.cpp b/Engine/Plugins/Interchange/Extensions/OpenUSD/Source/Import/Private/InterchangeUsdTranslator.cpp
--- a/Engine/Plugins/Interchange/Extensions/OpenUSD/Source/Import/Private/InterchangeUsdTranslator.cpp
+++ b/Engine/Plugins/Interchange/Extensions/OpenUSD/Source/Import/Private/InterchangeUsdTranslator.cpp
@@ -2954,7 +2954,7 @@ namespace UE::InterchangeUsdTranslator::Private
 		}
 	}
 
-	void AddSkeletalAnimationNode(
+	const UInterchangeSkeletalAnimationTrackNode* AddSkeletalAnimationNode(
 		const UE::FUsdSkelSkeletonQuery& SkeletonQuery,
 		const TMap<FString, TPair<FString, int32>>& BoneToUidAndBoneIndex,
 		UInterchangeUSDTranslatorImpl& TranslatorImpl,
@@ -2966,19 +2966,19 @@ namespace UE::InterchangeUsdTranslator::Private
 		UE::FUsdSkelAnimQuery AnimQuery = SkeletonQuery.GetAnimQuery();
 		if (!AnimQuery)
 		{
-			return;
+			return nullptr;
 		}
 
 		UE::FUsdPrim SkelAnimationPrim = AnimQuery.GetPrim();
 		if (!SkelAnimationPrim)
 		{
-			return;
+			return nullptr;
 		}
 
 		UE::FUsdPrim SkeletonPrim = SkeletonQuery.GetSkeleton();
 		if (!SkeletonPrim)
 		{
-			return;
+			return nullptr;
 		}
 
 		UE::FUsdStage Stage = SkeletonPrim.GetStage();
@@ -2992,7 +2992,7 @@ namespace UE::InterchangeUsdTranslator::Private
 		const UInterchangeSkeletalAnimationTrackNode* ExistingNode = Cast<UInterchangeSkeletalAnimationTrackNode>(NodeContainer.GetNode(NodeUid));
 		if (ExistingNode)
 		{
-			return;
+			return ExistingNode;
 		};
 
 		UInterchangeSkeletalAnimationTrackNode* SkelAnimNode = NewObject<UInterchangeSkeletalAnimationTrackNode>(&NodeContainer);
@@ -3084,7 +3084,7 @@ namespace UE::InterchangeUsdTranslator::Private
 			);
 			if (!bSuccess)
 			{
-				return;
+				return SkelAnimNode;
 			}
 
 			TArray<FString> SkelAnimChannelOrder = AnimQuery.GetBlendShapeOrder();
@@ -3170,6 +3170,8 @@ namespace UE::InterchangeUsdTranslator::Private
 				}
 			}
 		}
+
+		return SkelAnimNode;
 	}
 
 	void AddSkeletonNodes(
@@ -3292,7 +3294,18 @@ namespace UE::InterchangeUsdTranslator::Private
 		}
 
 		// Handle SkelAnimation prims, if we have any bound for this Skeleton
-		AddSkeletalAnimationNode(SkelQuery, BoneToUidAndBoneIndex, TranslatorImpl, SkeletonPrimNode, NodeContainer, Info);
+		const UInterchangeSkeletalAnimationTrackNode* SkelAnimNode = AddSkeletalAnimationNode(
+			SkelQuery,
+			BoneToUidAndBoneIndex,
+			TranslatorImpl,
+			SkeletonPrimNode,
+			NodeContainer,
+			Info
+		);
+		if (SkelAnimNode)
+		{
+			SkeletonPrimNode.SetCustomAnimationAssetUidToPlay(SkelAnimNode->GetUniqueID());
+		}
 
 		// Cache our joint names in order, as this is needed when generating skeletal mesh payloads
 		Info.SkelJointNames = MakeShared<TArray<FString>>();
