USD: Cache skeletal mesh descriptions to speed up morph target mesh description computation.

diff --git a/Engine/Plugins/Interchange/Extensions/OpenUSD/Source/Import/Private/InterchangeUsdTranslator.cpp b/Engine/Plugins/Interchange/Extensions/OpenUSD/Source/Import/Private/InterchangeUsdTranslator.cpp
--- a/Engine/Plugins/Interchange/Extensions/OpenUSD/Source/Import/Private/InterchangeUsdTranslator.cpp
+++ b/Engine/Plugins/Interchange/Extensions/OpenUSD/Source/Import/Private/InterchangeUsdTranslator.cpp
@@ -349,6 +349,9 @@ namespace UE::InterchangeUsdTranslator::Private
 		// Owned by the translator itself
 		TObjectPtr<UInterchangeResultsContainer> ResultsContainer;
 
+		mutable FRWLock SkeletalMeshDescriptionsLock;
+		TMap<FString, FMeshDescription> PayloadKeyToSkeletalMeshDescriptions;
+
 		// We store temp stages in here that we open in order to parse stuff inside of inactive variants of our main UsdStage.
 		// We do this because the payload data are retrieved concurrently, and toggling variants mutates the current stage.
 		mutable FRWLock PrimPathToVariantToStageLock;
@@ -3782,6 +3785,15 @@ namespace UE::InterchangeUsdTranslator::Private
 			}
 		}
 
+		// We cache these because we may need to retrieve these again when computing morph target mesh descriptions
+		FWriteScopeLock WriteLock{Impl.SkeletalMeshDescriptionsLock};
+		if (FMeshDescription* FoundMeshDescription = Impl.PayloadKeyToSkeletalMeshDescriptions.Find(PayloadKey))
+		{
+			OutMeshDescription = *FoundMeshDescription;
+			OutJointNames = *JointNames;
+			return true;
+		}
+
 		UE::FUsdSkelSkinningQuery SkinningQuery = UsdUtils::CreateSkinningQuery(Prim, SkelQuery);
 		if (!SkinningQuery)
 		{
@@ -3832,6 +3844,8 @@ namespace UE::InterchangeUsdTranslator::Private
 
 		OutJointNames = *JointNames;
 
+		Impl.PayloadKeyToSkeletalMeshDescriptions.Add(PayloadKey, OutMeshDescription);
+
 		return true;
 #else
 		return false;
