Interchange USD:
- Fix Render Context not correctly reset when Reset to Default button is hit

diff --git a/Engine/Plugins/Interchange/Extensions/OpenUSD/Source/Editor/Private/InterchangeUsdTranslatorSettingsCustomization.cpp b/Engine/Plugins/Interchange/Extensions/OpenUSD/Source/Editor/Private/InterchangeUsdTranslatorSettingsCustomization.cpp
--- a/Engine/Plugins/Interchange/Extensions/OpenUSD/Source/Editor/Private/InterchangeUsdTranslatorSettingsCustomization.cpp
+++ b/Engine/Plugins/Interchange/Extensions/OpenUSD/Source/Editor/Private/InterchangeUsdTranslatorSettingsCustomization.cpp
@@ -48,37 +48,6 @@ void FInterchangeUsdTranslatorSettingsCustomization::CustomizeDetails(IDetailLay
 		return;
 	}
 
-	// We used to pull these from UsdUnreal::MaterialUtils::GetRegisteredRenderContexts(), but these are
-	// hard-coded for Interchange because that list may contain render contexts added in through the
-	// legacy USD render context registry, which we don't necessarily support on Interchange just yet (e.g. "mdl")
-	TArray<FName> AllowedRenderContexts = {
-		TEXT("mtlx"),
-		TEXT("universal"),
-		TEXT("unreal"),
-	};
-
-	RenderContextComboBoxItems.Reset();
-	TSharedPtr<FString> InitiallySelectedContext;
-	for (const FName& Context : AllowedRenderContexts)
-	{
-		TSharedPtr<FString> ContextStr;
-		if (Context == UnrealIdentifiers::UniversalRenderContext)
-		{
-			ContextStr = MakeShared<FString>(UnrealIdentifiers::UniversalRenderContextDisplayString);
-		}
-		else
-		{
-			ContextStr = MakeShared<FString>(Context.ToString());
-		}
-
-		if (Context == CurrentOptions->RenderContext)
-		{
-			InitiallySelectedContext = ContextStr;
-		}
-
-		RenderContextComboBoxItems.Add(ContextStr);
-	}
-
 	IDetailCategoryBuilder& CatBuilder = DetailLayoutBuilder.EditCategory(TEXT("USD Translator"));
 
 	if (TSharedPtr<IPropertyHandle> RenderContextProperty = DetailLayoutBuilder.GetProperty(
@@ -98,51 +67,63 @@ void FInterchangeUsdTranslatorSettingsCustomization::CustomizeDetails(IDetailLay
 		]
 		.ValueContent()
 		[
-			SAssignNew(RenderContextComboBox, SComboBox<TSharedPtr<FString>>)
-			.OptionsSource(&RenderContextComboBoxItems)
-			.InitiallySelectedItem(InitiallySelectedContext)
-			.OnSelectionChanged_Lambda([this, RenderContextProperty](TSharedPtr<FString> NewContext, ESelectInfo::Type SelectType)
-			{
-				if (CurrentOptions == nullptr || !NewContext.IsValid())
-				{
-					return;
-				}
-
-				FScopedTransaction Transaction(LOCTEXT("RenderContextTransaction", "Edit Render Context"));
-				RenderContextProperty->NotifyPreChange();
-				{
-					FName NewContextName = (*NewContext) == UnrealIdentifiers::UniversalRenderContextDisplayString
-						? UnrealIdentifiers::UniversalRenderContext
-						: FName(**NewContext);
+			SNew(SBox)
+				.VAlign(VAlign_Center)
+				[
+					SNew(SComboBox<TSharedPtr<FString>>)
+						.OptionsSource(&RenderContextComboBoxItems)
+						.OnComboBoxOpening_Lambda([this]()
+						{
+							// We used to pull these from UsdUnreal::MaterialUtils::GetRegisteredRenderContexts(), but these are
+							// hard-coded for Interchange because that list may contain render contexts added in through the
+							// legacy USD render context registry, which we don't necessarily support on Interchange just yet (e.g. "mdl")
+							RenderContextComboBoxItems = {
+								MakeShared<FString>(TEXT("mtlx")),
+								MakeShared<FString>(TEXT("universal")),
+								MakeShared<FString>(TEXT("unreal")),
+							};
+						})
+						.OnGenerateWidget_Lambda([&](TSharedPtr<FString> Option)
+						{
+							TSharedPtr<SWidget> Widget = SNullWidget::NullWidget;
+							if (Option)
+							{
+								Widget = SNew(STextBlock)
+									.Text(FText::FromString(*Option))
+									.Font(FAppStyle::GetFontStyle("PropertyWindow.NormalFont"));
+							}
 
-					CurrentOptions->RenderContext = NewContextName;
-				}
-				RenderContextProperty->NotifyPostChange(EPropertyChangeType::ValueSet);
-				RenderContextProperty->NotifyFinishedChangingProperties();
-			})
-			.OnGenerateWidget_Lambda([](TSharedPtr<FString> Item)
-			{
-				return SNew(STextBlock)
-					.Text(Item.IsValid() ? FText::FromString(*Item) : FText::GetEmpty())
-					.Font(FAppStyle::GetFontStyle(TEXT("PropertyWindow.NormalFont")));
-			})
-			.Content()
-			[
-				SNew(STextBlock)
-				.Text_Lambda([this]() -> FText
-				{
-					TSharedPtr<FString> SelectedItem = RenderContextComboBox->GetSelectedItem();
-					if (SelectedItem.IsValid())
-					{
-						return FText::FromString(*SelectedItem);
-					}
+							return Widget.ToSharedRef();
+						})
+						.OnSelectionChanged_Lambda([this, RenderContextProperty](TSharedPtr<FString> ChosenOption, ESelectInfo::Type SelectInfo)
+						{
+							if (CurrentOptions && ChosenOption)
+							{
+								FScopedTransaction Transaction(LOCTEXT("RenderContextTransaction", "Edit Render Context"));
+								RenderContextProperty->NotifyPreChange();
+								{
+									CurrentOptions->RenderContext = **ChosenOption;
+								}
+								RenderContextProperty->NotifyPostChange(EPropertyChangeType::ValueSet);
+								RenderContextProperty->NotifyFinishedChangingProperties();
+							}
 
-					return FText::GetEmpty();
-				})
-				.Font(FAppStyle::GetFontStyle(TEXT("PropertyWindow.NormalFont")))
-			]
+						})
+						[
+							SNew(STextBlock)
+							.Text_Lambda([this]() -> FText
+							{
+								if (CurrentOptions)
+								{
+									return FText::FromString(CurrentOptions->RenderContext.ToString());
+								}
+
+								return FText::GetEmpty();
+							})
+							.Font(FAppStyle::GetFontStyle(TEXT("PropertyWindow.NormalFont")))
+						]
+				]
 		];
-		// clang-format on
 	}
 
 	if (TSharedPtr<IPropertyHandle> MaterialPurposeProperty = DetailLayoutBuilder.GetProperty(
@@ -208,7 +189,7 @@ void FInterchangeUsdTranslatorSettingsCustomization::CustomizeDetails(IDetailLay
 								? UnrealIdentifiers::MaterialAllPurposeText
 								: *Option
 							))
-							.Font(FAppStyle::GetFontStyle("PropertyWindow.NormalFont"));
+							.Font(FAppStyle::GetFontStyle(TEXT("PropertyWindow.NormalFont")));
 					}
 
 					return Widget.ToSharedRef();
diff --git a/Engine/Plugins/Interchange/Extensions/OpenUSD/Source/Editor/Private/InterchangeUsdTranslatorSettingsCustomization.h b/Engine/Plugins/Interchange/Extensions/OpenUSD/Source/Editor/Private/InterchangeUsdTranslatorSettingsCustomization.h
--- a/Engine/Plugins/Interchange/Extensions/OpenUSD/Source/Editor/Private/InterchangeUsdTranslatorSettingsCustomization.h
+++ b/Engine/Plugins/Interchange/Extensions/OpenUSD/Source/Editor/Private/InterchangeUsdTranslatorSettingsCustomization.h
@@ -24,7 +24,6 @@ public:
 private:
 	UInterchangeUsdTranslatorSettings* CurrentOptions;
 
-	TSharedPtr<SComboBox<TSharedPtr<FString>>> RenderContextComboBox;
 	TArray<TSharedPtr<FString>> RenderContextComboBoxItems;
 
 	TArray<TSharedPtr<FString>> MaterialPurposeComboBoxItems;
