USD: Fix Interchange USD not using timeCodesPerSecond when baking skeletal animation tracks.

diff --git a/Engine/Plugins/Interchange/Extensions/OpenUSD/Source/Import/Private/InterchangeUsdTranslator.cpp b/Engine/Plugins/Interchange/Extensions/OpenUSD/Source/Import/Private/InterchangeUsdTranslator.cpp
--- a/Engine/Plugins/Interchange/Extensions/OpenUSD/Source/Import/Private/InterchangeUsdTranslator.cpp
+++ b/Engine/Plugins/Interchange/Extensions/OpenUSD/Source/Import/Private/InterchangeUsdTranslator.cpp
@@ -2061,11 +2061,19 @@ namespace UE::InterchangeUsdTranslator::Private
 			int32 EndFrame = FMath::CeilToInt(Stage.GetEndTimeCode());
 			UsdUtils::GetAnimatedMeshTimeCodes(Stage, PrimPath, StartFrame, EndFrame);
 
-			double FramesPerSecond = Stage.GetTimeCodesPerSecond();
-			if (FramesPerSecond == 0)
+			double TimeCodesPerSecond = Stage.GetTimeCodesPerSecond();
+			if (TimeCodesPerSecond <= 0)
 			{
-				UE_LOG(LogUsd, Warning, TEXT("Invalid USD GeometryCache FPS detected. Falling back to 1 FPS"));
-				FramesPerSecond = 1;
+				UE_LOG(
+					LogUsd,
+					Warning,
+					TEXT(
+						"Stage '%s' has TimeCodesPerSecond set to '%f' which is not supported for GeometryCaches, which need values greater than zero. The GeometryCache assets will be parsed as if TimeCodesPerSecond was set to 1.0"
+					),
+					*Stage.GetRootLayer().GetIdentifier(),
+					TimeCodesPerSecond
+				);
+				TimeCodesPerSecond = 1;
 			}
 
 			// The GeometryCache module expects the end frame to be one past the last animation frame
@@ -2073,7 +2081,7 @@ namespace UE::InterchangeUsdTranslator::Private
 
 			GeometryCacheNode->SetCustomStartFrame(StartFrame);
 			GeometryCacheNode->SetCustomEndFrame(EndFrame);
-			GeometryCacheNode->SetCustomFrameRate(FramesPerSecond);
+			GeometryCacheNode->SetCustomFrameRate(TimeCodesPerSecond);
 
 			const bool bConstantTopology = UsdUtils::GetMeshTopologyVariance(Prim) != UsdUtils::EMeshTopologyVariance::Heterogenous;
 			GeometryCacheNode->SetCustomHasConstantTopology(bConstantTopology);
@@ -2800,8 +2808,9 @@ namespace UE::InterchangeUsdTranslator::Private
 
 		UInterchangeAnimationTrackSetNode* TrackSetNode = NewObject<UInterchangeAnimationTrackSetNode>(&NodeContainer);
 		TrackSetNode->InitializeNode(AnimTrackSetNodeUid, AnimTrackSetNodeDisplayName, EInterchangeNodeContainerType::TranslatedAsset);
-		TrackSetNode->SetCustomFrameRate(Layer.GetFramesPerSecond());	 // Key values in Interchange seem to be in seconds, so timeCodesPerSecond is
-																		 // not relevant here
+		
+		// This ends up as the LevelSequence frame rate, so it should probably match the stage's frame rate like legacy USD does
+		TrackSetNode->SetCustomFrameRate(Layer.GetFramesPerSecond());
 
 		NodeContainer.AddNode(TrackSetNode);
 		Impl.CurrentTrackSet = TrackSetNode;
@@ -3013,13 +3022,8 @@ namespace UE::InterchangeUsdTranslator::Private
 
 		// Time info
 		{
-			// TODO: Match the TrackSet framerate whenever Interchange supports skeletal animation sections.
-			// float TrackSetFrameRate = 30.0f;
-			// if (TranslatorImpl->CurrentTrackSet->GetCustomFrameRate(TrackSetFrameRate))
-			// {
-			// 	SkelAnimNode->SetCustomAnimationSampleRate(TrackSetFrameRate);
-			// }
-			SkelAnimNode->SetCustomAnimationSampleRate(Stage.GetFramesPerSecond());
+			const double TimeCodesPerSecond = Stage.GetTimeCodesPerSecond();			
+			SkelAnimNode->SetCustomAnimationSampleRate(TimeCodesPerSecond);
 
 			TOptional<double> StartTimeCode;
 			TOptional<double> StopTimeCode;
@@ -3040,8 +3044,6 @@ namespace UE::InterchangeUsdTranslator::Private
 				StopTimeCode = FMath::Max(BlendShapeTimeSamples[BlendShapeTimeSamples.Num() - 1], StopTimeCode.Get(TNumericLimits<double>::Lowest()));
 			}
 
-			UE::FUsdStage UsdStage = SkeletonPrim.GetStage();
-			double TimeCodesPerSecond = UsdStage.GetTimeCodesPerSecond();
 			if (StartTimeCode.IsSet())
 			{
 				SkelAnimNode->SetCustomAnimationStartTime(StartTimeCode.GetValue() / TimeCodesPerSecond);
