USD: Fix handling of the mesh baking transform for skeletal meshes from USD Interchange;

Use the skeleton scene transform + geomBindTransform when baking skeletal meshes;
Apply AdditionalTransform to morph target position offsets as well;

diff --git a/Engine/Plugins/Interchange/Extensions/OpenUSD/Source/Import/Private/InterchangeUsdTranslator.cpp b/Engine/Plugins/Interchange/Extensions/OpenUSD/Source/Import/Private/InterchangeUsdTranslator.cpp
--- a/Engine/Plugins/Interchange/Extensions/OpenUSD/Source/Import/Private/InterchangeUsdTranslator.cpp
+++ b/Engine/Plugins/Interchange/Extensions/OpenUSD/Source/Import/Private/InterchangeUsdTranslator.cpp
@@ -3859,8 +3859,6 @@ namespace UE::InterchangeUsdTranslator::Private
 			return false;
 		}
 
-		// TODO: This is extremely slow, as it will reimport the mesh for every single morph target!
-		// It seems to be what the other translators do, however. We need some form of FMeshDescription caching here
 		TArray<FString> UnusedJointNames;
 		const FString MeshPayloadKey = bIsLODMesh ? LODPrefix + MeshPrimPath : MeshPrimPath;
 		bool bConverted = GetSkeletalMeshPayloadData(MeshPayloadKey, Impl, Options, OutMeshDescription, UnusedJointNames);
@@ -4645,11 +4643,10 @@ TOptional<UE::Interchange::FMeshPayloadData> UInterchangeUSDTranslator::GetMeshP
 	using namespace UsdToUnreal;
 
 	UsdToUnreal::FUsdMeshConversionOptions OptionsCopy = ImplPtr->CachedMeshConversionOptions;
-	PayloadAttributes.GetAttribute(UE::Interchange::FAttributeKey{ MeshPayload::Attributes::MeshGlobalTransform }, OptionsCopy.AdditionalTransform);
 	OptionsCopy.bMergeIdenticalMaterialSlots = false;	 // This must always be false, because we need the material assignments we read from the
 														 // meshes to match up with whatever we cached from AddMeshNode, in order to fixup LOD
 														 // material slots
-		
+
 	static_assert(uint8(EInterchangeUsdPrimvar::All) == FUsdMeshConversionOptions::EImportPrimvar::All, "FUsdMeshConversionOptions::EImportPrimvar::All is different from EInterchangeUsdPrimvar::All");
 	static_assert(uint8(EInterchangeUsdPrimvar::Bake) == FUsdMeshConversionOptions::EImportPrimvar::Bake, "FUsdMeshConversionOptions::EImportPrimvar::Bake is different from EInterchangeUsdPrimvar::Bake");
 	static_assert(uint8(EInterchangeUsdPrimvar::Standard) == FUsdMeshConversionOptions::EImportPrimvar::Standard, "FUsdMeshConversionOptions::EImportPrimvar::Standard is different from EInterchangeUsdPrimvar::Standard");
@@ -4676,6 +4673,8 @@ TOptional<UE::Interchange::FMeshPayloadData> UInterchangeUSDTranslator::GetMeshP
 	{
 		case EInterchangeMeshPayLoadType::STATIC:
 		{
+			PayloadAttributes.GetAttribute(UE::Interchange::FAttributeKey{ MeshPayload::Attributes::MeshGlobalTransform }, OptionsCopy.AdditionalTransform);
+
 			bSuccess = UE::InterchangeUsdTranslator::Private::GetStaticMeshPayloadData(
 				PayloadKey.UniqueId,
 				*ImplPtr,
@@ -4686,6 +4685,19 @@ TOptional<UE::Interchange::FMeshPayloadData> UInterchangeUSDTranslator::GetMeshP
 		}
 		case EInterchangeMeshPayLoadType::SKELETAL:
 		{
+			// Don't use MeshGlobalTransform here as that will be the scene transform of our Mesh prims, which is not relevant for USD skinning.
+			// With baking, we want to first apply geomBindTransform, and then apply the skeleton's localToWorld transform. ConvertGeomMesh can 
+			// sort out the geomBindTransform (which should always be applied), so here we set the baking transform to the skeleton prim's transform if needed
+			bool bBakeMeshes = false;
+			FTransform RootJointGlobalTransform;
+			PayloadAttributes.GetAttribute(UE::Interchange::FAttributeKey{MeshPayload::Attributes::BakeMeshes}, bBakeMeshes);
+			PayloadAttributes.GetAttribute(UE::Interchange::FAttributeKey{MeshPayload::Attributes::RootJointGlobalTransform}, RootJointGlobalTransform);
+
+			if (bBakeMeshes)
+			{
+				OptionsCopy.AdditionalTransform = RootJointGlobalTransform;
+			}
+
 			bSuccess = UE::InterchangeUsdTranslator::Private::GetSkeletalMeshPayloadData(
 				PayloadKey.UniqueId,
 				*ImplPtr,
@@ -4697,6 +4709,17 @@ TOptional<UE::Interchange::FMeshPayloadData> UInterchangeUSDTranslator::GetMeshP
 		}
 		case EInterchangeMeshPayLoadType::MORPHTARGET:
 		{
+			// See the ::SKELETAL case
+			bool bBakeMeshes = false;
+			FTransform RootJointGlobalTransform;
+			PayloadAttributes.GetAttribute(UE::Interchange::FAttributeKey{MeshPayload::Attributes::BakeMeshes}, bBakeMeshes);
+			PayloadAttributes.GetAttribute(UE::Interchange::FAttributeKey{MeshPayload::Attributes::RootJointGlobalTransform}, RootJointGlobalTransform);
+
+			if (bBakeMeshes)
+			{
+				OptionsCopy.AdditionalTransform = RootJointGlobalTransform;
+			}
+
 			bSuccess = UE::InterchangeUsdTranslator::Private::GetMorphTargetPayloadData(
 				PayloadKey.UniqueId,
 				*ImplPtr,
@@ -4706,8 +4729,10 @@ TOptional<UE::Interchange::FMeshPayloadData> UInterchangeUSDTranslator::GetMeshP
 			);
 			break;
 		}
-		case EInterchangeMeshPayLoadType::ANIMATED:
+		case EInterchangeMeshPayLoadType::ANIMATED:	   // Geometry caches
 		{
+			PayloadAttributes.GetAttribute(UE::Interchange::FAttributeKey{ MeshPayload::Attributes::MeshGlobalTransform }, OptionsCopy.AdditionalTransform);
+
 			OptionsCopy.TimeCode = PayloadKey.FrameNumber;
 
 			bSuccess = UE::InterchangeUsdTranslator::Private::GetStaticMeshPayloadData(
diff --git a/Engine/Plugins/Interchange/Runtime/Source/Import/Private/Mesh/InterchangeSkeletalMeshFactory.cpp b/Engine/Plugins/Interchange/Runtime/Source/Import/Private/Mesh/InterchangeSkeletalMeshFactory.cpp
--- a/Engine/Plugins/Interchange/Runtime/Source/Import/Private/Mesh/InterchangeSkeletalMeshFactory.cpp
+++ b/Engine/Plugins/Interchange/Runtime/Source/Import/Private/Mesh/InterchangeSkeletalMeshFactory.cpp
@@ -882,7 +882,10 @@ void UInterchangeSkeletalMeshFactory::CreatePayloadTasks(const FImportAssetObjec
 		ImportAssetObjectLODData.LodIndex = CurrentLodIndex;
 
 		SkeletonNode->GetCustomUseTimeZeroForBindPose(ImportAssetObjectLODData.bUseTimeZeroAsBindPose);
-		
+
+		FTransform RootJointGlobalTransform;
+		RootJointNode->GetCustomGlobalTransform(Arguments.NodeContainer, GlobalOffsetTransform, RootJointGlobalTransform);
+
 		Private::BuildMeshReferences(
 			Arguments
 			, RootJointNodeId
@@ -926,6 +929,9 @@ void UInterchangeSkeletalMeshFactory::CreatePayloadTasks(const FImportAssetObjec
 		UE::Interchange::FAttributeStorage PayloadAttributes;
 		UInterchangeMeshFactoryNode::CopyPayloadKeyStorageAttributes(SkeletalMeshFactoryNode, PayloadAttributes);
 
+		PayloadAttributes.RegisterAttribute(UE::Interchange::FAttributeKey{ UE::Interchange::MeshPayload::Attributes::BakeMeshes }, bBakeMeshes);
+		PayloadAttributes.RegisterAttribute(UE::Interchange::FAttributeKey{ UE::Interchange::MeshPayload::Attributes::RootJointGlobalTransform }, RootJointGlobalTransform);
+
 		for (const FMeshNodeContext& MeshNodeContext : ImportAssetObjectLODData.MeshNodeContexts)
 		{
 			const FInternalInstanceData& InstanceData = MeshInstancesDatas.FindChecked(MeshNodeContext.TranslatorPayloadKey.UniqueId);
diff --git a/Engine/Plugins/Interchange/Runtime/Source/Nodes/Public/InterchangeMeshDefinitions.h b/Engine/Plugins/Interchange/Runtime/Source/Nodes/Public/InterchangeMeshDefinitions.h
--- a/Engine/Plugins/Interchange/Runtime/Source/Nodes/Public/InterchangeMeshDefinitions.h
+++ b/Engine/Plugins/Interchange/Runtime/Source/Nodes/Public/InterchangeMeshDefinitions.h
@@ -72,7 +72,9 @@ namespace UE
 		{
 			namespace Attributes
 			{
+				const FString RootJointGlobalTransform = TEXT("__Payload__RootJointGlobalTransform");
 				const FString MeshGlobalTransform = TEXT("__Payload__MeshGlobalTransform");
+				const FString BakeMeshes = TEXT("__Payload__BakeMeshes");
 			}
 		}
 	}
diff --git a/Engine/Plugins/Runtime/USDCore/Source/USDUtilities/Private/USDSkeletalDataConversion.cpp b/Engine/Plugins/Runtime/USDCore/Source/USDUtilities/Private/USDSkeletalDataConversion.cpp
--- a/Engine/Plugins/Runtime/USDCore/Source/USDUtilities/Private/USDSkeletalDataConversion.cpp
+++ b/Engine/Plugins/Runtime/USDCore/Source/USDUtilities/Private/USDSkeletalDataConversion.cpp
@@ -393,7 +393,7 @@ namespace SkelDataConversionImpl
 			// Position offset
 			// Note: TransformDir here because even the position deltas are still *deltas* (i.e. vector offsets)
 			pxr::GfVec3d USDPoint = GeomBindTransform.TransformDir(PointOffsets[OffsetIndex]);
-			ModifiedVertex.PositionDelta = (FVector3f)UsdToUnreal::ConvertVector(StageInfo, USDPoint);
+			ModifiedVertex.PositionDelta = (FVector3f)(FVector4f)TotalMatrix.TransformVector(UsdToUnreal::ConvertVector(StageInfo, USDPoint));
 
 			// Normal offset
 			if (OutBlendShape.bHasAuthoredTangents)
@@ -3405,6 +3405,9 @@ bool UsdUtils::ApplyBlendShape(
 
 	bool bApplied = false;
 
+	FMatrix TotalMatrix = AdditionalTransform.ToMatrixWithScale();
+	FMatrix TotalMatrixForNormal = TotalMatrix.Inverse().GetTransposed();
+
 	// We have one value for each vertex of the mesh description
 	// TODO: Deduplicate this code, since only the indexing differs?
 	if (PointIndices.size() == 0)
@@ -3414,7 +3417,7 @@ bool UsdUtils::ApplyBlendShape(
 		{
 			for (uint32 OffsetIndex = 0; OffsetIndex < PositionOffsets.size(); ++OffsetIndex)
 			{
-				const FVector UEOffset = UsdToUnreal::ConvertVector(StageInfo, PositionOffsets[OffsetIndex]);
+				const FVector UEOffset = TotalMatrix.TransformVector(UsdToUnreal::ConvertVector(StageInfo, PositionOffsets[OffsetIndex]));
 				MeshPositions[OffsetIndex] += FVector3f{UEOffset * Weight};
 			}
 
@@ -3435,9 +3438,6 @@ bool UsdUtils::ApplyBlendShape(
 		// Normal offsets
 		if (NormalOffsets.size() == (size_t)MeshPositions.GetNumElements())
 		{
-			FMatrix TotalMatrix = AdditionalTransform.ToMatrixWithScale();
-			FMatrix TotalMatrixForNormal = TotalMatrix.Inverse().GetTransposed();
-
 			for (uint32 VertexIndex = 0; VertexIndex < NormalOffsets.size(); ++VertexIndex)
 			{
 				const FVector UENormal = TotalMatrixForNormal.TransformVector(UsdToUnreal::ConvertVector(StageInfo, NormalOffsets[VertexIndex])).GetSafeNormal();
@@ -3474,7 +3474,7 @@ bool UsdUtils::ApplyBlendShape(
 				int TargetPointIndex = PointIndices[OffsetIndex];
 				if (TargetPointIndex >= 0 && TargetPointIndex < MeshPositions.GetNumElements())
 				{
-					const FVector UEOffset = UsdToUnreal::ConvertVector(StageInfo, PositionOffsets[OffsetIndex]);
+					const FVector UEOffset = TotalMatrix.TransformVector(UsdToUnreal::ConvertVector(StageInfo, PositionOffsets[OffsetIndex]));
 					MeshPositions[TargetPointIndex] += FVector3f{UEOffset * Weight};
 				}
 			}
@@ -3498,9 +3498,6 @@ bool UsdUtils::ApplyBlendShape(
 		// Normal offsets
 		if (PointIndices.size() == NormalOffsets.size())
 		{
-			FMatrix TotalMatrix = AdditionalTransform.ToMatrixWithScale();
-			FMatrix TotalMatrixForNormal = TotalMatrix.Inverse().GetTransposed();
-
 			for (uint32 NormalIndex = 0; NormalIndex < NormalOffsets.size(); ++NormalIndex)
 			{
 				int TargetPointIndex = PointIndices[NormalIndex];
