USD: Volume node handling fixes:

Fix SVT assets not being generated if exact same volume node shows up on the stage more than once (whether instancing or not);
Change volume node setup to have the scene nodes point at the asset nodes via CustomAssetUid, instead of the other way around;
Fix volume grid nodes not being properly parented to volume nodes as they're moved from a different node container.

diff --git a/Engine/Plugins/Interchange/Extensions/OpenUSD/Source/Import/Private/InterchangeUSDPipeline.cpp b/Engine/Plugins/Interchange/Extensions/OpenUSD/Source/Import/Private/InterchangeUSDPipeline.cpp
--- a/Engine/Plugins/Interchange/Extensions/OpenUSD/Source/Import/Private/InterchangeUSDPipeline.cpp
+++ b/Engine/Plugins/Interchange/Extensions/OpenUSD/Source/Import/Private/InterchangeUSDPipeline.cpp
@@ -488,32 +488,33 @@ namespace UE::InterchangeUsdPipeline::Private
 		using namespace UE::Interchange;
 		using namespace UE::Interchange::USD;
 
-		TMap<const UInterchangeSceneNode*, const UInterchangeVolumeNode*> SceneNodeToFirstVolume;
+		TSet<const UInterchangeSceneNode*> VolumeSceneNodes;
 		TMap<const UInterchangeVolumeNode*, UInterchangeSparseVolumeTextureFactoryNode*> FirstVolumeToFactoryNode;
 
-		// Collect all the node pairings we need
-		NodeContainer.IterateNodesOfType<UInterchangeSparseVolumeTextureFactoryNode>(
-			[&SceneNodeToFirstVolume,
-			 &NodeContainer,
-			 &FirstVolumeToFactoryNode](const FString& FactoryNodeUid, UInterchangeSparseVolumeTextureFactoryNode* FactoryNode)
+		NodeContainer.IterateNodes(
+			[&NodeContainer, &FirstVolumeToFactoryNode, &VolumeSceneNodes](const FString& FactoryNodeUid, UInterchangeBaseNode* BaseNode)
 			{
-				FString FirstVolumeUid = UInterchangeFactoryBaseNode::BuildTranslatedNodeUid(FactoryNodeUid);
-
-				const UInterchangeVolumeNode* FirstVolume = Cast<const UInterchangeVolumeNode>(NodeContainer.GetNode(FirstVolumeUid));
-				ensure(FirstVolume);
+				if (UInterchangeSparseVolumeTextureFactoryNode* VolumeFactoryNode = Cast<UInterchangeSparseVolumeTextureFactoryNode>(BaseNode))
+				{
+					FString FirstVolumeUid = UInterchangeFactoryBaseNode::BuildTranslatedNodeUid(FactoryNodeUid);
 
-				FirstVolumeToFactoryNode.Add(FirstVolume, FactoryNode);
+					const UInterchangeVolumeNode* FirstVolume = Cast<const UInterchangeVolumeNode>(NodeContainer.GetNode(FirstVolumeUid));
+					ensure(FirstVolume);
 
-				// Find our scene node too
-				TArray<FString> TargetNodeUids;
-				FirstVolume->GetTargetNodeUids(TargetNodeUids);
-				for (const FString& TargetNodeUid : TargetNodeUids)
+					FirstVolumeToFactoryNode.Add(FirstVolume, VolumeFactoryNode);
+				}
+				else if (UInterchangeSceneNode* SceneNode = Cast<UInterchangeSceneNode>(BaseNode))
 				{
-					const UInterchangeBaseNode* BaseNode = NodeContainer.GetNode(TargetNodeUid);
-					if (const UInterchangeSceneNode* SceneNode = Cast<const UInterchangeSceneNode>(BaseNode))
+					TArray<FString> TargetAssetUids;
+					SceneNode->GetTargetNodeUids(TargetAssetUids);
+
+					for (const FString& TargetUid : TargetAssetUids)
 					{
-						SceneNodeToFirstVolume.Add(SceneNode, FirstVolume);
-						break;
+						if (const UInterchangeVolumeNode* Volume = Cast<const UInterchangeVolumeNode>(NodeContainer.GetNode(TargetUid)))
+						{
+							VolumeSceneNodes.Add(SceneNode);
+							break;
+						}
 					}
 				}
 			}
@@ -526,9 +527,8 @@ namespace UE::InterchangeUsdPipeline::Private
 			MoveCustomAttributesToVolumeFactoryNode(VolumeNode, FactoryNode);
 		}
 
-		for (const TPair<const UInterchangeSceneNode*, const UInterchangeVolumeNode*>& Pair : SceneNodeToFirstVolume)
+		for (const UInterchangeSceneNode* SceneNode : VolumeSceneNodes)
 		{
-			const UInterchangeSceneNode* SceneNode = Pair.Key;
 			ProcessHeterogeneousVolumeSceneNode(SceneNode, NodeContainer);
 		}
 	}
diff --git a/Engine/Plugins/Interchange/Extensions/OpenUSD/Source/Import/Private/InterchangeUsdTranslator.cpp b/Engine/Plugins/Interchange/Extensions/OpenUSD/Source/Import/Private/InterchangeUsdTranslator.cpp
--- a/Engine/Plugins/Interchange/Extensions/OpenUSD/Source/Import/Private/InterchangeUsdTranslator.cpp
+++ b/Engine/Plugins/Interchange/Extensions/OpenUSD/Source/Import/Private/InterchangeUsdTranslator.cpp
@@ -2483,7 +2483,6 @@ namespace UE::InterchangeUsdTranslator::Private
 		// separate factory nodes. If we want to keep the expected mapping of factory node / volume node Uids (just an added "Factory_" prefix)
 		// this means we need a separate volume node *per animation*, so we'll also use "AnimationIDs" to differentiate them
 
-		TStrongObjectPtr<UInterchangeBaseNodeContainer> TempContainer = nullptr;
 		TStrongObjectPtr<UInterchangeTranslatorBase> Translator = nullptr;
 
 		// This is collected by path hash here because for animated SVTs we want to still have a single UsdUtils::FVolumePrimInfo for each
@@ -2541,6 +2540,9 @@ namespace UE::InterchangeUsdTranslator::Private
 			// frame of animation, or the default opinion if we have that
 			UInterchangeVolumeNode* FirstVolumeNode = nullptr;
 
+			TSet<UInterchangeVolumeNode*> ResetVolumeNodes;
+			ResetVolumeNodes.Reserve(VDBFilePaths.Num());
+
 			for (int32 Index = 0; Index < VDBFilePaths.Num(); ++Index)
 			{
 				const FString& FilePath = VDBFilePaths[Index];
@@ -2597,62 +2599,29 @@ namespace UE::InterchangeUsdTranslator::Private
 						Settings->AnimationID = VolumePrimInfo.TimeSamplePathTimeCodes.Num() > 0 ? AnimationID : FString{};
 					}
 
-					if (!TempContainer)
-					{
-						TempContainer.Reset(NewObject<UInterchangeBaseNodeContainer>());
-					}
-					TempContainer->Reset();	   // Empty the container of nodes
-
 					Translator->SourceData = SourceData;
-					Translator->Translate(*TempContainer);
+					Translator->Translate(InNodeContainer);
 
 					// Cache that we used this translator for this filepath. If we keep it, we don't have to open the file again to
 					// retrieve the payload. Note: This is likely the same translator we used for all .vdb files during this import
 					Translators.Add(FilePath, Translator);
 
-					// Move the nodes into our own node container. We do this because we must iterate to discover the volume node
-					// (as we have no guarantees about its NodeUid), and we don't want to do a full loop over our own NodeContainer
-					// that may have other Volume nodes and many more nodes overall.
-					//
-					// Also note the slight exploit: We can get non-const access to the translated nodes in this way, which we need
-					// in order to call InNodeContainer::Add
-					TempContainer->IterateNodes(
-						[&InNodeContainer, &VolumeNode](const FString& NodeUid, UInterchangeBaseNode* Node)
+					// Note the slight exploit: We can get non-const access to the translated nodes in this way, which we need
+					InNodeContainer.BreakableIterateNodesOfType<UInterchangeVolumeNode>(
+						[&InNodeContainer, &VolumeNode, &FilePath](const FString& NodeUid, UInterchangeVolumeNode* Node)
 						{
-							InNodeContainer.AddNode(Node);
-
-							if (UInterchangeVolumeNode* CastNode = Cast<UInterchangeVolumeNode>(Node))
+							if (NodeUid.Contains(FilePath))
 							{
-								if (VolumeNode == nullptr)
-								{
-									// Note that since we're not discovering new .vdbs, we expect exactly one UInterchangeVolumeNode to be generated
-									VolumeNode = CastNode;
-								}
-								else
-								{
-									UE_LOG(
-										LogUsd,
-										Warning,
-										TEXT("Found unexpected volume node '%s' (display label '%s')"),
-										*NodeUid,
-										*CastNode->GetDisplayLabel()
-									);
-								}
+								VolumeNode = Node;
+								return true;
 							}
+
+							return false;
 						}
 					);
 
 					if (VolumeNode)
 					{
-						// Remove any animation index that the OpenVDB translator may have set on this volume node,
-						// as we want to control that explicitly from here
-						TArray<int32> ExistingAnimationIndices;
-						VolumeNode->GetCustomFrameIndicesInAnimation(ExistingAnimationIndices);
-						for (const int32 ExistingIndex : ExistingAnimationIndices)
-						{
-							VolumeNode->RemoveCustomFrameIndexInAnimation(ExistingIndex);
-						}
-
 						AnimationIDToVolumeNode.Add(AnimationID, VolumeNode);
 					}
 				}
@@ -2663,6 +2632,19 @@ namespace UE::InterchangeUsdTranslator::Private
 					continue;
 				}
 
+				// Remove any animation index that the OpenVDB translator (or a translation of another instance of this whole volume prim in the
+				// stage) may have set on this volume node, as we want to set the animation indices exclusively during this pass through VDBFilePaths
+				if (!ResetVolumeNodes.Contains(VolumeNode))
+				{
+					ResetVolumeNodes.Add(VolumeNode);
+
+					TArray<int32> ExistingAnimationIndices;
+					VolumeNode->GetCustomFrameIndicesInAnimation(ExistingAnimationIndices);
+					for (const int32 ExistingIndex : ExistingAnimationIndices)
+					{
+						VolumeNode->RemoveCustomFrameIndexInAnimation(ExistingIndex);
+					}
+				}
 				VolumeNode->AddCustomFrameIndexInAnimation(Index);
 
 				if (!FirstVolumeNode)
@@ -2917,7 +2899,7 @@ namespace UE::InterchangeUsdTranslator::Private
 		{
 			if (const UInterchangeVolumeNode* VolumeNode = Cast<const UInterchangeVolumeNode>(NodeContainer.GetNode(AssetNodeUid)))
 			{
-				VolumeNode->AddTargetNodeUid(SceneNodeUid);
+				SceneNode->AddTargetNodeUid(AssetNodeUid);
 			}
 		}
 
