USD Interchange: Use instancing-aware translation.

diff --git a/Engine/Plugins/Interchange/Extensions/OpenUSD/Source/Import/Private/InterchangeUSDPipeline.cpp b/Engine/Plugins/Interchange/Extensions/OpenUSD/Source/Import/Private/InterchangeUSDPipeline.cpp
--- a/Engine/Plugins/Interchange/Extensions/OpenUSD/Source/Import/Private/InterchangeUSDPipeline.cpp
+++ b/Engine/Plugins/Interchange/Extensions/OpenUSD/Source/Import/Private/InterchangeUSDPipeline.cpp
@@ -6,6 +6,7 @@
 #include "InterchangeMaterialDefinitions.h"
 #include "InterchangeMaterialFactoryNode.h"
 #include "InterchangeMaterialInstanceNode.h"
+#include "InterchangeMeshActorFactoryNode.h"
 #include "InterchangeMeshFactoryNode.h"
 #include "InterchangeMeshNode.h"
 #include "InterchangeSceneNode.h"
@@ -823,9 +824,10 @@ namespace UE::InterchangeUsdPipeline::Private
 	// Makes sure that all materials in the MeshNode's SlotDependencies array are primvar-compatible
 	// with the primvar to UV index mapping of the mesh, creating new materials and reassigning them to the
 	// SlotDependencies custom attribute if needed.
+	template<typename NodeType>
 	void GeneratePrimvarCompatibleMaterials(
 		const UInterchangeMeshNode* InMeshNode,
-		UInterchangeMeshFactoryNode* InMeshFactoryNode,
+		NodeType* InMeshFactoryNode,
 		UInterchangeBaseNodeContainer* InNodeContainer,
 		TMap<FString, FMaterialPrimvarInfo>& InOutFactoryNodeUidToInfo
 	)
@@ -951,10 +953,11 @@ namespace UE::InterchangeUsdPipeline::Private
 
 		// Collect the nodes in a separate container because GeneratePrimvarCompatibleMaterials may add additional nodes
 		// to the node container itself, which we can't do while we iterate over it
-		TMap<const UInterchangeMeshNode*, UInterchangeMeshFactoryNode*> MeshNodeToFactoryNode;
+		TMap<const UInterchangeMeshNode*, UInterchangeMeshFactoryNode*> MeshNodeToMeshFactoryNode;
+		TMap<const UInterchangeMeshNode*, UInterchangeMeshActorFactoryNode*> MeshNodeToMeshActorFactoryNode;
 
 		NodeContainer->IterateNodesOfType<UInterchangeMeshNode>(
-			[NodeContainer, &MeshNodeToFactoryNode, ImportPrimvar](const FString& NodeUid, UInterchangeMeshNode* MeshNode)
+			[NodeContainer, &MeshNodeToMeshFactoryNode, ImportPrimvar](const FString& NodeUid, UInterchangeMeshNode* MeshNode)
 			{
 				const FString FactoryNodeUid = UInterchangeFactoryBaseNode::BuildFactoryNodeUid(NodeUid);
 				UInterchangeMeshFactoryNode* MeshFactoryNode = Cast<UInterchangeMeshFactoryNode>(NodeContainer->GetFactoryNode(FactoryNodeUid));
@@ -962,13 +965,11 @@ namespace UE::InterchangeUsdPipeline::Private
 				{
 					return;
 				}
-
-				MeshFactoryNode->AddPayloadKeyInt32Attribute(UE::Interchange::USD::Primvar::Import, int32(ImportPrimvar));
-
-				MeshNodeToFactoryNode.Add(MeshNode, MeshFactoryNode);
+				MeshNodeToMeshFactoryNode.Add(MeshNode, MeshFactoryNode);
 
 #if WITH_EDITORONLY_DATA
-				if(ImportPrimvar != EInterchangeUsdPrimvar::Standard)
+				MeshFactoryNode->AddPayloadKeyInt32Attribute(UE::Interchange::USD::Primvar::Import, int32(ImportPrimvar));
+				if (ImportPrimvar != EInterchangeUsdPrimvar::Standard)
 				{
 					AddPrimvarAttributesToFactoryNodes(MeshNode, MeshFactoryNode, NodeContainer);
 				}
@@ -978,13 +979,51 @@ namespace UE::InterchangeUsdPipeline::Private
 
 		if (bGenerateCompatibleMaterials)
 		{
-			// Build this map to prevent us from reparsing the same material in case multiple meshes
-			// are using them
+			// Also search through our mesh actor factory nodes: We may have added material overrides that also need to be made compatible.
+			//
+			// Two passes through the node container is not ideal, but at least we can turn off this whole pass by disabling
+			// bGenerateCompatibleMaterials, as opposed to having the collection of these nodes always running on the above pass
+			NodeContainer->IterateNodesOfType<UInterchangeMeshActorFactoryNode>(
+				[NodeContainer, &MeshNodeToMeshActorFactoryNode](const FString& NodeUid, UInterchangeMeshActorFactoryNode* MeshActorFactoryNode)
+				{
+					TMap<FString, FString> MaterialDependencies;
+					MeshActorFactoryNode->GetSlotMaterialDependencies(MaterialDependencies);
+					if (MaterialDependencies.Num() == 0)
+					{
+						return;
+					}
+
+					FString MeshFactoryNodeUid;
+					MeshActorFactoryNode->GetCustomInstancedAssetFactoryNodeUid(MeshFactoryNodeUid);
+					UInterchangeMeshFactoryNode* MeshFactoryNode = Cast<UInterchangeMeshFactoryNode>(NodeContainer->GetFactoryNode(MeshFactoryNodeUid));
+					if (!MeshFactoryNode)
+					{
+						return;
+					}
+
+					TArray<FString> MeshFactoryNodeTargetUids;
+					MeshFactoryNode->GetTargetNodeUids(MeshFactoryNodeTargetUids);
+					for (const FString& TargetUid : MeshFactoryNodeTargetUids)
+					{
+						if (const UInterchangeMeshNode* MeshNode = Cast<const UInterchangeMeshNode>(NodeContainer->GetNode(TargetUid)))
+						{
+							MeshNodeToMeshActorFactoryNode.Add(MeshNode, MeshActorFactoryNode);
+						}
+					}
+				}
+			);
+
+			// Build this map to prevent us from reparsing the same material in case multiple meshes are using them
 			TMap<FString, FMaterialPrimvarInfo> MaterialFactoryNodeUidToInfo;
-			for (const TPair<const UInterchangeMeshNode*, UInterchangeMeshFactoryNode*>& MeshPair : MeshNodeToFactoryNode)
+
+			for (const TPair<const UInterchangeMeshNode*, UInterchangeMeshFactoryNode*>& MeshPair : MeshNodeToMeshFactoryNode)
 			{
 				GeneratePrimvarCompatibleMaterials(MeshPair.Key, MeshPair.Value, NodeContainer, MaterialFactoryNodeUidToInfo);
 			}
+			for (const TPair<const UInterchangeMeshNode*, UInterchangeMeshActorFactoryNode*>& ActorPair : MeshNodeToMeshActorFactoryNode)
+			{
+				GeneratePrimvarCompatibleMaterials(ActorPair.Key, ActorPair.Value, NodeContainer, MaterialFactoryNodeUidToInfo);
+			}
 		}
 	}
 
diff --git a/Engine/Plugins/Interchange/Extensions/OpenUSD/Source/Import/Private/InterchangeUsdTranslator.cpp b/Engine/Plugins/Interchange/Extensions/OpenUSD/Source/Import/Private/InterchangeUsdTranslator.cpp
--- a/Engine/Plugins/Interchange/Extensions/OpenUSD/Source/Import/Private/InterchangeUsdTranslator.cpp
+++ b/Engine/Plugins/Interchange/Extensions/OpenUSD/Source/Import/Private/InterchangeUsdTranslator.cpp
@@ -340,6 +340,14 @@ namespace UE::InterchangeUsdTranslator::Private
 			const UInterchangeSceneNode* SceneNodeWithLODs
 		);
 
+		template<typename MeshOrSceneNode>
+		void SetSlotMaterialDependencies(
+			MeshOrSceneNode* MeshNode,
+			const UsdUtils::FUsdPrimMaterialAssignmentInfo& MaterialAssignments,
+			UInterchangeUsdTranslatorSettings* TranslatorSettings,
+			UInterchangeBaseNodeContainer& NodeContainer
+		);
+
 	public:
 		// We have to keep a stage reference so that we can parse the payloads after Translate() complete.
 		// ReleaseSource() clears this member, once translation is complete.
@@ -873,7 +881,7 @@ namespace UE::InterchangeUsdTranslator::Private
 		UInterchangeBaseNodeContainer& NodeContainer
 	)
 	{
-		FString PrimPath = Prim.GetPrimPath().GetString();
+		FString PrimPath = UsdUtils::GetPrototypePrimPath(Prim).GetString();
 		FString NodeName{FPaths::GetCleanFilename(Value.TextureFilePath)};
 
 		// Check if Node already exist with this ID
@@ -1049,7 +1057,7 @@ namespace UE::InterchangeUsdTranslator::Private
 	// Sets on the provided MeshNode custom attributes needed to bake the provided GeomPropValue geomprops / primvars
 	// into textures, for the provided MaterialNodeUid
 	void AddPrimvarBakingAttributes(
-		UInterchangeMeshNode* MeshNode,
+		UInterchangeBaseNode* BaseNode,
 		const FString& MaterialNodeUid,
 		UInterchangeBaseNodeContainer& NodeContainer,
 		const TArray<FUsdMaterialXShaderGraph::FGeomProp>& GeomPropValues
@@ -1058,6 +1066,25 @@ namespace UE::InterchangeUsdTranslator::Private
 #if WITH_EDITOR
 		using namespace UE::Interchange;
 
+		// TODO: Is it OK that this function is called repeatedly for the same Mesh? (e.g. from multiple material assignments,
+		// from multiple scene node instancing the same Mesh asset node, etc.)
+
+		UInterchangeMeshNode* MeshNode = Cast<UInterchangeMeshNode>(BaseNode);
+		if (!MeshNode)
+		{
+			// We also use this function when setting material overrides on scene nodes. In that case our scene node should have a
+			// Mesh asset node UID
+			if (UInterchangeSceneNode* SceneNode = Cast<UInterchangeSceneNode>(BaseNode))
+			{
+				FString AssetInstanceUid;
+				SceneNode->GetCustomAssetInstanceUid(AssetInstanceUid);
+
+				if (const UInterchangeMeshNode* SceneMeshNode = Cast<UInterchangeMeshNode>(NodeContainer.GetNode(AssetInstanceUid)))
+				{
+					MeshNode = const_cast<UInterchangeMeshNode*>(SceneMeshNode);
+				}
+			}
+		}
 		if (!MeshNode)
 		{
 			return;
@@ -1212,7 +1239,7 @@ namespace UE::InterchangeUsdTranslator::Private
 			// Emit texture node itself (this is the main place where this happens)
 			// Note that the node name isn't just the texture path, as we may have multiple material users of this texture
 			// with different settings, and so we need separate translated nodes for each material and parameter
-			const FString TextureUid = FString::Printf(TEXT("Texture:%s:%s"), *Prim.GetPrimPath().GetString(), **BaseParameterName);
+			const FString TextureUid = FString::Printf(TEXT("Texture:%s:%s"), *UsdUtils::GetPrototypePrimPath(Prim).GetString(), **BaseParameterName);
 			AddTextureNode(Prim, TextureUid, Value, NodeContainer);
 
 			// Actual texture assignment
@@ -1335,7 +1362,7 @@ namespace UE::InterchangeUsdTranslator::Private
 			}
 		}
 
-		FString PrimPath = Prim.GetPrimPath().GetString();
+		const FString PrimPath = UsdUtils::GetPrototypePrimPath(Prim).GetString();
 		FString MaterialPrimName = Prim.GetName().ToString();
 		FString MaterialUid = MaterialPrefix + PrimPath;
 		FString MaterialNodeName = MaterialPrimName;
@@ -1714,6 +1741,70 @@ namespace UE::InterchangeUsdTranslator::Private
 		CurrentLODSceneNodes.Reset();
 	}
 
+	template<typename MeshOrSceneNodeType>
+	void UInterchangeUSDTranslatorImpl::SetSlotMaterialDependencies(
+		MeshOrSceneNodeType* MeshOrSceneNode,
+		const UsdUtils::FUsdPrimMaterialAssignmentInfo& MaterialAssignments,
+		UInterchangeUsdTranslatorSettings* TranslatorSettings,
+		UInterchangeBaseNodeContainer& NodeContainer
+	)
+	{
+		for (const UsdUtils::FUsdPrimMaterialSlot& Slot : MaterialAssignments.Slots)
+		{
+			// We do this because Interchange will, in some scenarios, merge material slots with identical slot names.
+			// By using the source (which is the displaycolor desc / material prim path / unreal material content path)
+			// we do end up with a goofy looking super long material slot names, but it will have Interchange only combine
+			// slots if they really are pointing at the exact same thing
+			FString SlotName = Slot.MaterialSource;
+
+			// Get the Uid of the material instance that we'll end up assigning to this slot
+			FString MaterialUid;
+			switch (Slot.AssignmentType)
+			{
+				case UsdUtils::EPrimAssignmentType::DisplayColor:
+				{
+					// MaterialSource here is e.g. "!DisplayColor_0_1"
+					MaterialUid = AddDisplayColorMaterialInstanceNodeIfNeeded(NodeContainer, Slot);
+					break;
+				}
+				case UsdUtils::EPrimAssignmentType::MaterialPrim:
+				{
+					// MaterialSource here is the material prim path
+					UE::FUsdPrim MaterialPrim = UsdStage.GetPrimAtPath(UE::FSdfPath{*Slot.MaterialSource});
+
+					// Go through a prototype in case this material is inside an instance proxy, as this needs to match up with
+					// the polygon group name we're going to add to our mesh description.
+					UE::FUsdPrim MaterialPrototypePrim = UsdUtils::GetPrototypePrim(MaterialPrim);
+					if (MaterialPrototypePrim != MaterialPrim)
+					{
+						SlotName = MaterialPrototypePrim.GetPrimPath().GetString();
+					}
+
+					MaterialUid = AddMaterialNode(MaterialPrototypePrim, TranslatorSettings, NodeContainer, Slot.bMeshIsDoubleSided);
+
+					if (TArray<FUsdMaterialXShaderGraph::FGeomProp>* FoundGeomProps = MaterialUidToGeomProps.Find(MaterialUid))
+					{
+						AddPrimvarBakingAttributes(MeshOrSceneNode, MaterialUid, NodeContainer, *FoundGeomProps);
+					}
+					break;
+				}
+				case UsdUtils::EPrimAssignmentType::UnrealMaterial:
+				{
+					// MaterialSource here is the content path, e.g. "/Game/MyFolder/Red.Red"
+					MaterialUid = AddUnrealMaterialReferenceNodeIfNeeded(NodeContainer, Slot.MaterialSource);
+					break;
+				}
+				default:
+				{
+					ensure(false);
+					break;
+				}
+			}
+
+			MeshOrSceneNode->SetSlotMaterialDependencyUid(SlotName, *MaterialUid);
+		}
+	}
+
 	FString AddLightNode(const UE::FUsdPrim& Prim, UInterchangeBaseNodeContainer& NodeContainer)
 	{
 		TRACE_CPUPROFILER_EVENT_SCOPE(AddLightNode)
@@ -1896,7 +1987,7 @@ namespace UE::InterchangeUsdTranslator::Private
 			return;
 		}
 
-		const FString MeshPrimPath = MeshPrim.GetPrimPath().GetString();
+		const FString MeshPrimPath = UsdUtils::GetPrototypePrimPath(MeshPrim).GetString();
 
 		TFunction<void(const FString&, int32, const FString&)> AddMorphTargetNode =
 			[&MeshNode, &MeshPrimPath, &NodeContainer, &Info](const FString& MorphTargetName, int32 BlendShapeIndex, const FString& InbetweenName)
@@ -1984,7 +2075,7 @@ namespace UE::InterchangeUsdTranslator::Private
 	}
 
 	FString UInterchangeUSDTranslatorImpl::AddMeshNode(
-		const UE::FUsdPrim& Prim,
+		const UE::FUsdPrim& InstanceProxyPrim,
 		UInterchangeUsdTranslatorSettings* TranslatorSettings,
 		UInterchangeBaseNodeContainer& NodeContainer,
 		const FTraversalInfo& Info,
@@ -1993,11 +2084,29 @@ namespace UE::InterchangeUsdTranslator::Private
 	{
 		TRACE_CPUPROFILER_EVENT_SCOPE(UInterchangeUSDTranslatorImpl::AddMeshNode)
 
-		FString PrimPath = Prim.GetPrimPath().GetString();
-		FString NodeUid = (bPrimitiveShape ? PrimitiveShapePrefix : MeshPrefix) + PrimPath;
-		FString NodeName(Prim.GetName().ToString());
+		const bool bIsSkinned = static_cast<bool>(Info.ClosestParentSkelRootPath) && InstanceProxyPrim.HasAPI(TEXT("SkelBindingAPI"));
 
-		UE::FUsdStage Stage = Prim.GetStage();
+		UE::FUsdPrim PrototypePrim;
+		if (bIsSkinned || Info.bInsideLOD)
+		{
+			// Don't try to reuse asset nodes across instances for variant-style LODs: USD generates the prototype
+			// paths on-demand, and as we flip through LODs one at a time, and later retrieve LOD mesh payload data
+			// by opening other stages with population masks, there's zero guarantee our prototype paths will match
+			//
+			// Also temporarily disable instancing for skinned meshes as Interchange doesn't support that very well
+			// yet (see UE-175330)
+			PrototypePrim = InstanceProxyPrim;
+		}
+		else
+		{
+			// Go through the prototype (if any) so that we try reusing the same mesh asset node for different USD instances
+			PrototypePrim = UsdUtils::GetPrototypePrim(InstanceProxyPrim);
+		}
+
+		const UE::FUsdStage Stage = PrototypePrim.GetStage();
+		const FString PrototypePrimPath = PrototypePrim.GetPrimPath().GetString();
+		const FString NodeUid = (bPrimitiveShape ? PrimitiveShapePrefix : MeshPrefix) + PrototypePrimPath;
+		const FString NodeName(PrototypePrim.GetName().ToString());
 
 		// Check if Node already exist with this ID
 		if (const UInterchangeMeshNode* Node = Cast<UInterchangeMeshNode>(NodeContainer.GetNode(NodeUid)))
@@ -2006,14 +2115,13 @@ namespace UE::InterchangeUsdTranslator::Private
 		}
 
 		// Fill in the MeshNode itself
-		const bool bIsAnimated = UsdUtils::IsAnimatedMesh(Prim);
+		const bool bIsAnimated = UsdUtils::IsAnimatedMesh(PrototypePrim);
 		UInterchangeMeshNode* MeshNode = bIsAnimated ? NewObject<UInterchangeGeometryCacheNode>(&NodeContainer)
 													 : NewObject<UInterchangeMeshNode>(&NodeContainer);
 		NodeContainer.SetupNode(MeshNode, NodeUid, NodeName, EInterchangeNodeContainerType::TranslatedAsset);
 		MeshNode->SetAssetName(NodeName);
-		const bool bIsSkinned = static_cast<bool>(Info.ClosestParentSkelRootPath) && Prim.HasAPI(TEXT("SkelBindingAPI"));
 
-		FString PayloadKey = PrimPath;
+		FString PayloadKey = PrototypePrimPath;
 		if (Info.bInsideLOD)
 		{
 			PayloadKey = LODPrefix + PayloadKey;
@@ -2032,14 +2140,13 @@ namespace UE::InterchangeUsdTranslator::Private
 			PayloadKey = PrimitiveShapePrefix + PayloadKey;
 		}
 
-		double TimeCode = UsdUtils::GetDefaultTimeCode();
 		if (UInterchangeGeometryCacheNode* GeometryCacheNode = Cast<UInterchangeGeometryCacheNode>(MeshNode))
 		{
-			GeometryCacheNode->SetPayLoadKey(PrimPath, EInterchangeMeshPayLoadType::ANIMATED);
+			GeometryCacheNode->SetPayLoadKey(PrototypePrimPath, EInterchangeMeshPayLoadType::ANIMATED);
 
 			int32 StartFrame = FMath::FloorToInt(Stage.GetStartTimeCode());
 			int32 EndFrame = FMath::CeilToInt(Stage.GetEndTimeCode());
-			UsdUtils::GetAnimatedMeshTimeCodes(Stage, PrimPath, StartFrame, EndFrame);
+			UsdUtils::GetAnimatedMeshTimeCodes(Stage, PrototypePrimPath, StartFrame, EndFrame);
 
 			double TimeCodesPerSecond = Stage.GetTimeCodesPerSecond();
 			if (TimeCodesPerSecond <= 0)
@@ -2063,10 +2170,8 @@ namespace UE::InterchangeUsdTranslator::Private
 			GeometryCacheNode->SetCustomEndFrame(EndFrame);
 			GeometryCacheNode->SetCustomFrameRate(TimeCodesPerSecond);
 
-			const bool bConstantTopology = UsdUtils::GetMeshTopologyVariance(Prim) != UsdUtils::EMeshTopologyVariance::Heterogenous;
+			const bool bConstantTopology = UsdUtils::GetMeshTopologyVariance(PrototypePrim) != UsdUtils::EMeshTopologyVariance::Heterogenous;
 			GeometryCacheNode->SetCustomHasConstantTopology(bConstantTopology);
-
-			TimeCode = StartFrame;
 		}
 		else if (bIsSkinned && !bPrimitiveShape)
 		{
@@ -2077,7 +2182,7 @@ namespace UE::InterchangeUsdTranslator::Private
 				MeshNode->SetSkeletonDependencyUid(*Info.BoundSkeletonPrimPath);
 			}
 
-			AddMorphTargetNodes(Prim, *this, *MeshNode, NodeContainer, Info);
+			AddMorphTargetNodes(PrototypePrim, *this, *MeshNode, NodeContainer, Info);
 
 			// When returning the payload data later, we'll need at the very least our SkeletonQuery, so
 			// here we store the Info object into the Impl
@@ -2090,7 +2195,7 @@ namespace UE::InterchangeUsdTranslator::Private
 		{
 			MeshNode->SetPayLoadKey(PayloadKey, EInterchangeMeshPayLoadType::STATIC);
 
-			if (UsdUtils::IsCollisionEnabledForPrim(Prim))
+			if (UsdUtils::IsCollisionEnabledForPrim(PrototypePrim))
 			{
 				// If the mesh prim is flagged for collision schemas AND also setup to be an FBX-style collision mesh, prefer
 				// the FBX style. We do this here because if we setup both styles at the same time, then GetCollisionMeshType() from
@@ -2118,10 +2223,10 @@ namespace UE::InterchangeUsdTranslator::Private
 					// In that case we want to disable the collider for this mesh prim itself, so that it matches legacy USD behavior
 					if (bSetCustomCollisionType)
 					{
-						TArray<UE::FUsdPrim> Siblings = Prim.GetParent().GetChildren();
+						TArray<UE::FUsdPrim> Siblings = PrototypePrim.GetParent().GetChildren();
 						for (const UE::FUsdPrim& Sibling : Siblings)
 						{
-							if (Sibling == Prim || !UsdUtils::IsCollisionMesh(Sibling))
+							if (Sibling == PrototypePrim || !UsdUtils::IsCollisionMesh(Sibling))
 							{
 								continue;
 							}
@@ -2144,7 +2249,7 @@ namespace UE::InterchangeUsdTranslator::Private
 
 				if (bSetCustomCollisionType)
 				{
-					EUsdCollisionType Approximation = UsdUtils::GetCollisionApproximationType(Prim);
+					EUsdCollisionType Approximation = UsdUtils::GetCollisionApproximationType(PrototypePrim);
 					EInterchangeMeshCollision InterchangeApproximation = ConvertApproximationType(Approximation);
 					if (InterchangeApproximation != EInterchangeMeshCollision::None)
 					{
@@ -2156,9 +2261,10 @@ namespace UE::InterchangeUsdTranslator::Private
 
 		// Material assignments
 		{
+			double TimeCode = UsdUtils::GetDefaultTimeCode();
 			const bool bProvideMaterialIndices = false;
-			UsdUtils::FUsdPrimMaterialAssignmentInfo Assignments = UsdUtils::GetPrimMaterialAssignments(
-				Prim,
+			UsdUtils::FUsdPrimMaterialAssignmentInfo PrototypeAssignments = UsdUtils::GetPrimMaterialAssignments(
+				PrototypePrim,
 				TimeCode,
 				bProvideMaterialIndices,
 				CachedMeshConversionOptions.RenderContext,
@@ -2167,7 +2273,7 @@ namespace UE::InterchangeUsdTranslator::Private
 
 			if (Info.bInsideLOD)
 			{
-				CachedMaterialAssignments.Add(PrimPath, Assignments);
+				CachedMaterialAssignments.Add(PrototypePrimPath, PrototypeAssignments);
 			}
 
 			// Move these into the asset node because the USD Pipeline will compare these with the assigned material's
@@ -2178,7 +2284,7 @@ namespace UE::InterchangeUsdTranslator::Private
 			// primvar-to-UVIndex mapping. It should though, as the mesh conversion options are the same. We can't cache
 			// the mapping because we run into USD allocator issues, given that all the strings contained in the
 			// FUsdPrimMaterialAssignmentInfo object are allocated inside an USD allocator scope
-			for (const TPair<FString, int32>& PrimvarPair : UsdUtils::GetPrimvarToUVIndexMap(Prim))
+			for (const TPair<FString, int32>& PrimvarPair : UsdUtils::GetPrimvarToUVIndexMap(PrototypePrim))
 			{
 				const FString& PrimvarName = PrimvarPair.Key;
 				int32 UVIndex = PrimvarPair.Value;
@@ -2186,56 +2292,12 @@ namespace UE::InterchangeUsdTranslator::Private
 				MeshNode->AddInt32Attribute(UE::Interchange::USD::PrimvarUVIndexAttributePrefix + PrimvarName, UVIndex);
 			}
 
-			for (const UsdUtils::FUsdPrimMaterialSlot& Slot : Assignments.Slots)
-			{
-				// We do this because Interchange will, in some scenarios, merge material slots with identical slot names.
-				// By using the source (which is the displaycolor desc / material prim path / unreal material content path)
-				// we do end up with a goofy looking super long material slot names, but it will have Interchange only combine
-				// slots if they really are pointing at the exact same thing
-				const FString& SlotName = Slot.MaterialSource;
-
-				// Get the Uid of the material instance that we'll end up assigning to this slot
-				FString MaterialUid;
-				switch (Slot.AssignmentType)
-				{
-					case UsdUtils::EPrimAssignmentType::DisplayColor:
-					{
-						// MaterialSource here is e.g. "!DisplayColor_0_1"
-						MaterialUid = AddDisplayColorMaterialInstanceNodeIfNeeded(NodeContainer, Slot);
-						break;
-					}
-					case UsdUtils::EPrimAssignmentType::MaterialPrim:
-					{
-						// MaterialSource here is the material prim path
-						UE::FUsdPrim MaterialPrim = Stage.GetPrimAtPath(UE::FSdfPath{*Slot.MaterialSource});
-						MaterialUid = AddMaterialNode(MaterialPrim, TranslatorSettings, NodeContainer, Slot.bMeshIsDoubleSided);
-						break;
-					}
-					case UsdUtils::EPrimAssignmentType::UnrealMaterial:
-					{
-						// MaterialSource here is the content path, e.g. "/Game/MyFolder/Red.Red"
-						MaterialUid = AddUnrealMaterialReferenceNodeIfNeeded(NodeContainer, Slot.MaterialSource);
-						break;
-					}
-					default:
-					{
-						ensure(false);
-						break;
-					}
-				}
-
-				MeshNode->SetSlotMaterialDependencyUid(SlotName, *MaterialUid);
-
-				if (TArray<FUsdMaterialXShaderGraph::FGeomProp>* FoundGeomProps = MaterialUidToGeomProps.Find(MaterialUid))
-				{
-					AddPrimvarBakingAttributes(MeshNode, MaterialUid, NodeContainer, *FoundGeomProps);
-				}
-			}
+			SetSlotMaterialDependencies(MeshNode, PrototypeAssignments, TranslatorSettings, NodeContainer);
 		}
 		return NodeUid;
 	}
 
-	void UInterchangeUSDTranslatorImpl::AddLODMeshNodes(	//
+	void UInterchangeUSDTranslatorImpl::AddLODMeshNodes(
 		const UE::FUsdPrim& Prim,
 		UInterchangeBaseNodeContainer& NodeContainer,
 		UInterchangeSceneNode& ParentSceneNode,
@@ -2950,7 +3012,7 @@ namespace UE::InterchangeUsdTranslator::Private
 		UE::FUsdStage Stage = SkeletonPrim.GetStage();
 
 		const FString SkelAnimationName = SkelAnimationPrim.GetName().ToString();
-		const FString SkelAnimationPrimPath = SkelAnimationPrim.GetPrimPath().GetString();
+		const FString SkelAnimationPrimPath = UsdUtils::GetPrototypePrimPath(SkelAnimationPrim).GetString();
 		const FString SkeletonPrimPath = SkeletonPrim.GetPrimPath().GetString();
 		const FString UniquePath = SkelAnimationPrimPath + TEXT("\\") + SkeletonPrimPath;
 		const FString NodeUid = AnimationTrackPrefix + UniquePath;
@@ -3064,7 +3126,7 @@ namespace UE::InterchangeUsdTranslator::Private
 				{
 					continue;
 				}
-				const FString MeshPrimPath = Prim.GetPrimPath().GetString();
+				const FString MeshPrimPath = UsdUtils::GetPrototypePrimPath(Prim).GetString();
 
 				TArray<FString> BlendShapeChannels;
 				bool bInnerSucces = SkinningTarget.GetBlendShapeOrder(BlendShapeChannels);
@@ -3416,6 +3478,9 @@ namespace UE::InterchangeUsdTranslator::Private
 		// Do this before generating other nodes as they may need the updated info
 		Info.UpdateWithCurrentPrim(Prim);
 
+		// TODO: Maybe refactor this function to create the scene node first, and then pass it into the AddXNode functions?
+		// That way they can assign their assets to the scene nodes however they want to, an we don't have to store these
+		// random return types here
 		FString VolumeMaterialInstanceUid;
 		bool bNeedsVolumeTrack = false;
 
@@ -3528,6 +3593,12 @@ namespace UE::InterchangeUsdTranslator::Private
 				SceneNode->SetCustomActorVisibility(Info.bVisible);
 			}
 
+			// Connect scene node and primary asset node
+			if (AssetNodeUids.Num() != 0)
+			{
+				SceneNode->SetCustomAssetInstanceUid(AssetNodeUids[0]);
+			}
+
 			if (Info.bIsLODContainer)
 			{
 				TranslatorImpl.AddLODMeshNodes(Prim, NodeContainer, *SceneNode, TranslatorSettings, Info);
@@ -3557,10 +3628,29 @@ namespace UE::InterchangeUsdTranslator::Private
 				);
 			}
 
-			// Connect scene node and primary asset node
-			if (AssetNodeUids.Num() != 0)
-			{
-				SceneNode->SetCustomAssetInstanceUid(AssetNodeUids[0]);
+			// Set component material overrides
+			//
+			// For now we only need material overrides in case of USD instanceables.
+			//
+			// TODO: While the USD translator and pipeline's side of this should be working, material overrides don't quite work just yet,
+			// as ActorHelper::ApplySlotMaterialDependencies (the thing responsible for setting material overrides on the actor components)
+			// will try to set these overrides according to slot name, while we always just place the (mesh's) material path itself as the
+			// slot name... We need some different mechanism for setting slot name (maybe a translator/pipeline setting to just use the
+			// GeomSubset prim name, or the slot index?) that can also work for setting material overrides. Also, see FixMaterialSlotNames()
+			if (Prim.IsA(TEXT("Mesh")) && (Prim.IsInstance() || Prim.IsInstanceProxy()))
+			{
+				// This is in charge of picking up on collection-based material assignments, as we'll use the instance proxy prim here
+				const double TimeCode = UsdUtils::GetDefaultTimeCode();
+				const bool bProvideMaterialIndices = false;
+				UsdUtils::FUsdPrimMaterialAssignmentInfo InstanceAssignments = UsdUtils::GetPrimMaterialAssignments(
+					Prim,
+					TimeCode,
+					bProvideMaterialIndices,
+					TranslatorImpl.CachedMeshConversionOptions.RenderContext,
+					TranslatorImpl.CachedMeshConversionOptions.MaterialPurpose
+				);
+
+				TranslatorImpl.SetSlotMaterialDependencies(SceneNode, InstanceAssignments, TranslatorSettings, NodeContainer);
 			}
 
 			// Add animation tracks
@@ -3623,8 +3713,6 @@ namespace UE::InterchangeUsdTranslator::Private
 		FStaticMeshAttributes StaticMeshAttributes(TempMeshDescription);
 		StaticMeshAttributes.Register();
 
-		// TODO: We can't do much with these yet: They will be used to generate primvar-compatible
-		// versions of the materials that are assigned to this mesh, whenever we get a pipeline
 		UsdUtils::FUsdPrimMaterialAssignmentInfo TempMaterialInfo;
 
 		bool bSuccess = false;
@@ -3694,7 +3782,7 @@ namespace UE::InterchangeUsdTranslator::Private
 			}
 		}
 
-		const FString& MeshNodeUid = MeshPrefix + Prim.GetPrimPath().GetString();
+		const FString& MeshNodeUid = MeshPrefix + UsdUtils::GetPrototypePrimPath(Prim).GetString();
 
 		// Read these variables from the data we cached during traversal for translation
 		TSharedPtr<TArray<FString>> JointNames = nullptr;
@@ -3862,7 +3950,7 @@ namespace UE::InterchangeUsdTranslator::Private
 			{
 				FReadScopeLock ReadLock{Impl.CachedTraversalInfoLock};
 
-				const FString& MeshNodeUid = MeshPrefix + MeshPrim.GetPrimPath().GetString();
+				const FString& MeshNodeUid = MeshPrefix + UsdUtils::GetPrototypePrimPath(MeshPrim).GetString();
 				const FTraversalInfo* MeshInfo = Impl.NodeUidToCachedTraversalInfo.Find(MeshNodeUid);
 				if (MeshInfo)
 				{
diff --git a/Engine/Plugins/Interchange/Runtime/Source/Pipelines/Private/InterchangeGenericGeometryCachePipeline.cpp b/Engine/Plugins/Interchange/Runtime/Source/Pipelines/Private/InterchangeGenericGeometryCachePipeline.cpp
--- a/Engine/Plugins/Interchange/Runtime/Source/Pipelines/Private/InterchangeGenericGeometryCachePipeline.cpp
+++ b/Engine/Plugins/Interchange/Runtime/Source/Pipelines/Private/InterchangeGenericGeometryCachePipeline.cpp
@@ -13,36 +13,60 @@ void UInterchangeGenericMeshPipeline::ExecutePreImportPipelineGeometryCache()
 {
 	check(CommonMeshesProperties.IsValid());
 
-	if (bImportGeometryCaches && CommonMeshesProperties->ForceAllMeshAsType == EInterchangeForceMeshType::IFMT_None)
+	if (!bImportGeometryCaches || CommonMeshesProperties->ForceAllMeshAsType != EInterchangeForceMeshType::IFMT_None)
 	{
-		// If there's an animated mesh, combine all meshes into a single Geometry Cache
-		TArray<FString> MeshInstanceUids;
-		PipelineMeshesUtilities->GetAllGeometryCacheInstance(MeshInstanceUids);
+		return;
+	}
 
-		TArray<FString> MeshUids;
-		for (const FString& MeshUid : MeshInstanceUids)
-		{
-			const FInterchangeMeshInstance& MeshInstance = PipelineMeshesUtilities->GetMeshInstanceByUid(MeshUid);
+	// If there's an animated mesh, combine all meshes into a single Geometry Cache
+	TArray<FString> MeshInstanceUids;
+	PipelineMeshesUtilities->GetAllGeometryCacheInstance(MeshInstanceUids);
+
+	TSet<FString> VisitedMeshAssetUids;
+
+	TArray<FString> MeshUids;
+	for (const FString& MeshUid : MeshInstanceUids)
+	{
+		const FInterchangeMeshInstance& MeshInstance = PipelineMeshesUtilities->GetMeshInstanceByUid(MeshUid);
 
-			// Only look at lod 0 since GeometryCache don't support LODs
-			const int32 LodIndex = 0;
-			FInterchangeLodSceneNodeContainer EmptySceneNodeContainer;
-			const FInterchangeLodSceneNodeContainer& SceneNodeContainer = MeshInstance.SceneNodePerLodIndex.Contains(LodIndex) ?
-				MeshInstance.SceneNodePerLodIndex.FindChecked(LodIndex) : EmptySceneNodeContainer;
+		// Only look at lod 0 since GeometryCache don't support LODs
+		const int32 LodIndex = 0;
+		FInterchangeLodSceneNodeContainer EmptySceneNodeContainer;
+		const FInterchangeLodSceneNodeContainer& SceneNodeContainer = MeshInstance.SceneNodePerLodIndex.Contains(LodIndex) ?
+			MeshInstance.SceneNodePerLodIndex.FindChecked(LodIndex) : EmptySceneNodeContainer;
 
-			for (const UInterchangeSceneNode* SceneNode : SceneNodeContainer.SceneNodes)
+		for (const UInterchangeSceneNode* SceneNode : SceneNodeContainer.SceneNodes)
+		{
+			if (CommonMeshesProperties->bBakeMeshes)
 			{
+				// If we're baking, always add all scene nodes: We may be trying to bake multiple instances of the same mesh node
+				// into a single big geometry cache asset
 				MeshUids.Add(SceneNode->GetUniqueID());
 			}
+			else
+			{
+				// If not baking, only add a single scene node per each mesh asset: This prevents us from ending up with multiple copies
+				// of the same mesh on top of itself (as we won't bake transforms), and it also allows us to instance geometry caches
+				// by just making several scene nodes that reference the same animated mesh node: Only one of the "instances" will end
+				// up on the geometry cache asset, and we'll reuse the asset for all the scene nodes
+				FString AssetUid;
+				SceneNode->GetCustomAssetInstanceUid(AssetUid);
+				if (!VisitedMeshAssetUids.Contains(AssetUid))
+				{
+					VisitedMeshAssetUids.Add(AssetUid);
+					MeshUids.Add(SceneNode->GetUniqueID());
+				}
+			}
 		}
+	}
 
-		// If we got some instances, create a geometry cache factory node
-		if (MeshUids.Num() > 0)
-		{
-			UInterchangeGeometryCacheFactoryNode* GeometryCacheFactoryNode = CreateGeometryCacheFactoryNode(MeshUids);
-			GeometryCacheFactoryNodes.Add(GeometryCacheFactoryNode);
-		}
+	if (MeshUids.Num() == 0)
+	{
+		return;
 	}
+
+	UInterchangeGeometryCacheFactoryNode* GeometryCacheFactoryNode = CreateGeometryCacheFactoryNode(MeshUids);
+	GeometryCacheFactoryNodes.Add(GeometryCacheFactoryNode);
 }
 
 UInterchangeGeometryCacheFactoryNode* UInterchangeGenericMeshPipeline::CreateGeometryCacheFactoryNode(const TArray<FString>& MeshUids)
diff --git a/Engine/Plugins/Runtime/USDCore/Source/USDUtilities/Private/Objects/USDSchemaTranslator.cpp b/Engine/Plugins/Runtime/USDCore/Source/USDUtilities/Private/Objects/USDSchemaTranslator.cpp
--- a/Engine/Plugins/Runtime/USDCore/Source/USDUtilities/Private/Objects/USDSchemaTranslator.cpp
+++ b/Engine/Plugins/Runtime/USDCore/Source/USDUtilities/Private/Objects/USDSchemaTranslator.cpp
@@ -7,6 +7,7 @@
 #include "USDMaterialUtils.h"
 #include "USDMemory.h"
 #include "USDTypesConversion.h"
+#include "USDConversionUtils.h"
 #include "UsdWrappers/UsdPrim.h"
 #include "UsdWrappers/UsdTyped.h"
 
@@ -28,15 +29,6 @@
 
 int32 FRegisteredSchemaTranslatorHandle::CurrentSchemaTranslatorId = 0;
 
-static bool GInstancingAwareTranslation = true;
-static FAutoConsoleVariableRef CVarInstancingAwareTranslation(
-	TEXT("USD.InstancingAwareTranslation"),
-	GInstancingAwareTranslation,
-	TEXT(
-		"Enabling this lets the USDImporter skip some extra steps during translation when it encounters multiple instance prims of the same (static) Mesh prototype prim."
-	)
-);
-
 FUsdSchemaTranslatorRegistry& FUsdSchemaTranslatorRegistry::Get()
 {
 	static FUsdSchemaTranslatorRegistry Singleton;
@@ -284,30 +276,12 @@ bool FUsdSchemaTranslator::IsCollapsed(ECollapsingType CollapsingType) const
 
 UE::FSdfPath FUsdSchemaTranslator::GetPrototypePrimPath() const
 {
-	if (!GInstancingAwareTranslation)
-	{
-		return PrimPath;
-	}
-
-	UE::FUsdPrim Prim = GetPrim();
-	if (Prim)
-	{
-		if (Prim.IsInstance())
-		{
-			return Prim.GetPrototype().GetPrimPath();
-		}
-		else if (Prim.IsInstanceProxy())
-		{
-			return Prim.GetPrimInPrototype().GetPrimPath();
-		}
-	}
-
-	return PrimPath;
+	return UsdUtils::GetPrototypePrimPath(GetPrim());
 }
 
 bool FUsdSchemaTranslator::ShouldSkipInstance() const
 {
-	if (!GInstancingAwareTranslation)
+	if (!UsdUtils::IsInstancingAwareTranslationEnabled())
 	{
 		return false;
 	}
diff --git a/Engine/Plugins/Runtime/USDCore/Source/USDUtilities/Private/USDConversionUtils.cpp b/Engine/Plugins/Runtime/USDCore/Source/USDUtilities/Private/USDConversionUtils.cpp
--- a/Engine/Plugins/Runtime/USDCore/Source/USDUtilities/Private/USDConversionUtils.cpp
+++ b/Engine/Plugins/Runtime/USDCore/Source/USDUtilities/Private/USDConversionUtils.cpp
@@ -137,6 +137,15 @@ static FAutoConsoleVariableRef CVarRemoveDuplicates(
 	)
 );
 
+static bool GInstancingAwareTranslation = true;
+static FAutoConsoleVariableRef CVarInstancingAwareTranslation(
+	TEXT("USD.InstancingAwareTranslation"),
+	GInstancingAwareTranslation,
+	TEXT(
+		"Enabling this lets the USDImporter skip some extra steps during translation when it encounters multiple instance prims of the same (static) Mesh prototype prim."
+	)
+);
+
 namespace USDConversionUtilsImpl
 {
 	/** Show some warnings if the UVSet primvars show some unsupported/problematic behavior */
@@ -4242,6 +4251,37 @@ void UsdUtils::ReadStageMetaData(UE::FUsdStage Stage, TMap<FString, FString>& Ou
 #endif
 }
 
+bool UsdUtils::IsInstancingAwareTranslationEnabled()
+{
+#if USE_USD_SDK
+	return GInstancingAwareTranslation;
+#else
+	return true;
+#endif
+}
+
+UE::FUsdPrim UsdUtils::GetPrototypePrim(const UE::FUsdPrim& Prim)
+{
+	if (!IsInstancingAwareTranslationEnabled())
+	{
+		return Prim;
+	}
+	if (Prim.IsInstance())
+	{
+		return Prim.GetPrototype();
+	}
+	else if (Prim.IsInstanceProxy())
+	{
+		return Prim.GetPrimInPrototype();
+	}
+	return Prim;
+}
+
+UE::FSdfPath UsdUtils::GetPrototypePrimPath(const UE::FUsdPrim& Prim)
+{
+	return GetPrototypePrim(Prim).GetPrimPath();
+}
+
 #if USE_USD_SDK
 #undef LOCTEXT_NAMESPACE
 #endif
diff --git a/Engine/Plugins/Runtime/USDCore/Source/USDUtilities/Public/USDConversionUtils.h b/Engine/Plugins/Runtime/USDCore/Source/USDUtilities/Public/USDConversionUtils.h
--- a/Engine/Plugins/Runtime/USDCore/Source/USDUtilities/Public/USDConversionUtils.h
+++ b/Engine/Plugins/Runtime/USDCore/Source/USDUtilities/Public/USDConversionUtils.h
@@ -701,4 +701,14 @@ namespace UsdUtils
 
 	/** Parses and stores primary Stage metadata key-info pairs in the MetadataMap */
 	USDUTILITIES_API void ReadStageMetaData(UE::FUsdStage Stage, TMap<FString, FString>& OutMetaDataMap);
+
+	// Whether or not instancing-aware translation is enabled
+	USDUTILITIES_API bool IsInstancingAwareTranslationEnabled();
+
+	// Retrieves the equivalent of this prim in the prototype it's referenced from, where applicable
+	USDUTILITIES_API UE::FUsdPrim GetPrototypePrim(const UE::FUsdPrim& Prim);
+	
+	// Retrieves the prototype path for the prim
+	USDUTILITIES_API UE::FSdfPath GetPrototypePrimPath(const UE::FUsdPrim& Prim);
+
 }	 // namespace UsdUtils
