Interchange MaterialX:
- Add support for the UDIMs

diff --git a/Engine/Plugins/Interchange/Extensions/OpenUSD/Source/Import/Private/InterchangeUsdTranslator.cpp b/Engine/Plugins/Interchange/Extensions/OpenUSD/Source/Import/Private/InterchangeUsdTranslator.cpp
--- a/Engine/Plugins/Interchange/Extensions/OpenUSD/Source/Import/Private/InterchangeUsdTranslator.cpp
+++ b/Engine/Plugins/Interchange/Extensions/OpenUSD/Source/Import/Private/InterchangeUsdTranslator.cpp
@@ -55,7 +55,7 @@
 #include "Volume/InterchangeVolumeDefinitions.h"
 #include "Volume/InterchangeVolumeTranslatorSettings.h"
 
-#include "MaterialX/MaterialXUtils/MaterialXBase.h"
+#include "MaterialX/InterchangeMaterialXTranslator.h"
 #include "MaterialX/MaterialXUtils/MaterialXManager.h"
 
 #include "Async/ParallelFor.h"
@@ -1167,6 +1167,7 @@ namespace UE::InterchangeUsdTranslator::Private
 					// check on the Translator, it might return nullptr in case of reimport
 					if (Translator)
 					{
+						Translator->Results = ResultsContainer;
 						Translator->Translate(NodeContainer);
 						Translators.Emplace(File, Translator);
 					}
@@ -4826,53 +4827,6 @@ TOptional<UE::Interchange::FMeshPayloadData> UInterchangeUSDTranslator::GetMeshP
 	return Result;
 }
 
-namespace UE::InterchangeUsdTranslator::Private
-{
-	TOptional<UE::Interchange::FImportImage> GetMaterialXTexturePayloadData(
-		const FString& PayloadKey,
-		TOptional<FString>& AlternateTexturePath,
-		const UInterchangeTranslatorBase* Translator,
-		UInterchangeUSDTranslatorImpl& Impl,
-		UInterchangeResultsContainer* Results,
-		TSharedPtr<UE::Interchange::FAnalyticsHelper> AnalyticsHelper
-	)
-	{
-		TRACE_CPUPROFILER_EVENT_SCOPE(GetMaterialXTexturePayloadData)
-
-		TOptional<UE::Interchange::FImportImage> TexturePayloadData;
-
-#if USE_USD_SDK
-		FString Filename = PayloadKey;
-		TextureCompressionSettings CompressionSettings = TextureCompressionSettings::TC_Default;
-
-#if WITH_EDITOR
-		if (int32 IndexTextureCompression; PayloadKey.FindChar(FMaterialXManager::TexturePayloadSeparator, IndexTextureCompression))
-		{
-			Filename = PayloadKey.Mid(0, IndexTextureCompression);
-			CompressionSettings = TextureCompressionSettings(FCString::Atoi(*PayloadKey.Mid(IndexTextureCompression + 1)));
-		}
-#endif
-		UE::Interchange::Private::FScopedTranslator ScopedTranslator(Filename, Results, AnalyticsHelper);
-		const IInterchangeTexturePayloadInterface* TextureTranslator = ScopedTranslator.GetPayLoadInterface<IInterchangeTexturePayloadInterface>();
-
-		if (!TextureTranslator)
-		{
-			return TexturePayloadData;
-		}
-
-		AlternateTexturePath = Impl.GetTextureSourcePath(Filename);
-
-		TexturePayloadData = TextureTranslator->GetTexturePayloadData(PayloadKey, AlternateTexturePath);
-		if (TexturePayloadData.IsSet())
-		{
-			TexturePayloadData.GetValue().CompressionSettings = CompressionSettings;
-		}
-#endif	  // USE_USD_SDK
-
-		return TexturePayloadData;
-	}
-}	 // namespace UE::InterchangeUsdTranslator::Private
-
 TOptional<UE::Interchange::FImportImage> UInterchangeUSDTranslator::GetTexturePayloadData(
 	const FString& PayloadKey,
 	TOptional<FString>& AlternateTexturePath
@@ -4893,6 +4847,7 @@ TOptional<UE::Interchange::FImportImage> UInterchangeUSDTranslator::GetTexturePa
 
 	FString FilePath;
 	TextureGroup TextureGroup;
+	AlternateTexturePath = ImplPtr->GetTextureSourcePath(FilePath);
 	bool bDecoded = DecodeTexturePayloadKey(PayloadKey, FilePath, TextureGroup);
 	if (bDecoded)
 	{
@@ -4901,8 +4856,6 @@ TOptional<UE::Interchange::FImportImage> UInterchangeUSDTranslator::GetTexturePa
 		const IInterchangeTexturePayloadInterface* TextureTranslator = ScopedTranslator.GetPayLoadInterface<IInterchangeTexturePayloadInterface>();
 		if (TextureTranslator)
 		{
-			AlternateTexturePath = ImplPtr->GetTextureSourcePath(FilePath);
-
 			// The texture translators don't use the payload key, and read the texture directly from the SourceData's file path
 			const FString UnusedPayloadKey = {};
 			TexturePayloadData = TextureTranslator->GetTexturePayloadData(UnusedPayloadKey, AlternateTexturePath);
@@ -4934,11 +4887,10 @@ TOptional<UE::Interchange::FImportImage> UInterchangeUSDTranslator::GetTexturePa
 		}
 	}
 
-	// If we couldn't find a texture in either the USD Translator nor the Translators, then it's most likely coming from reading an mtlx in memory,
-	// copying the behavior from InterchangeMaterialXTranslator
+	// If we couldn't find a texture in either the USD Translator nor the Translators, then it's most likely coming from reading an mtlx in memory
 	if (!TexturePayloadData.IsSet())
 	{
-		TexturePayloadData = GetMaterialXTexturePayloadData(PayloadKey, AlternateTexturePath, this, *ImplPtr, Results, AnalyticsHelper);
+		TexturePayloadData = UInterchangeMaterialXTranslator::GetTexturePayloadData(PayloadKey, AlternateTexturePath, Results, AnalyticsHelper);
 	}
 
 #endif	  // USE_USD_SDK
@@ -4965,13 +4917,35 @@ TOptional<UE::Interchange::FImportBlockedImage> UInterchangeUSDTranslator::GetBl
 
 	FString FilePath;
 	TextureGroup TextureGroup;
+	AlternateTexturePath = ImplPtr->GetTextureSourcePath(FilePath);
 	bool bDecoded = DecodeTexturePayloadKey(PayloadKey, FilePath, TextureGroup);
 	if (!bDecoded)
 	{
-		return {};
+		TOptional<UE::Interchange::FImportBlockedImage> BlockedTexturePayloadData;
+		// We did not find a suitable Payload in USD Translator, let's find one in one of the Translators (MaterialX for the moment)
+		// The best way would be to have a direct association between the payload and the right Translator, but we don't have a suitable way of knowing
+		// which Payload belongs to which Translator So let's just loop over them all
+		for (const TPair<FString, TStrongObjectPtr<UInterchangeTranslatorBase>>& Pair : Impl->Translators)
+		{
+			if (IInterchangeBlockedTexturePayloadInterface* TexturePayloadInterface = Cast<IInterchangeBlockedTexturePayloadInterface>(Pair.Value.Get()))
+			{
+				BlockedTexturePayloadData = TexturePayloadInterface->GetBlockedTexturePayloadData(PayloadKey, AlternateTexturePath);
+				if (BlockedTexturePayloadData.IsSet())
+				{
+					break;
+				}
+			}
+		}
+
+		// If we couldn't find a texture in either the USD Translator nor the Translators, then it's most likely coming from reading an mtlx in memory
+		if (!BlockedTexturePayloadData.IsSet())
+		{
+			BlockedTexturePayloadData = UInterchangeMaterialXTranslator::GetBlockedTexturePayloadData(PayloadKey, AlternateTexturePath, Results, AnalyticsHelper);
+		}
+
+		return BlockedTexturePayloadData;
 	}
 
-	AlternateTexturePath = ImplPtr->GetTextureSourcePath(FilePath);
 
 	// Collect all the UDIM tile filepaths similar to this current tile. If we've been asked to translate
 	// a blocked texture then we must have some
diff --git a/Engine/Plugins/Interchange/Runtime/Source/Import/Private/MaterialX/InterchangeMaterialXTranslator.cpp b/Engine/Plugins/Interchange/Runtime/Source/Import/Private/MaterialX/InterchangeMaterialXTranslator.cpp
--- a/Engine/Plugins/Interchange/Runtime/Source/Import/Private/MaterialX/InterchangeMaterialXTranslator.cpp
+++ b/Engine/Plugins/Interchange/Runtime/Source/Import/Private/MaterialX/InterchangeMaterialXTranslator.cpp
@@ -1,10 +1,10 @@
 // Copyright Epic Games, Inc. All Rights Reserved. 
 
 #include "MaterialX/InterchangeMaterialXTranslator.h"
-#include "InterchangeManager.h"
 #include "InterchangeTranslatorHelper.h"
 #include "MaterialX/MaterialXUtils/MaterialXManager.h"
 #include "Nodes/InterchangeSourceNode.h"
+#include "UDIMUtilities.h"
 
 #include UE_INLINE_GENERATED_CPP_BY_NAME(InterchangeMaterialXTranslator)
 
@@ -60,22 +60,27 @@ bool UInterchangeMaterialXTranslator::Translate(UInterchangeBaseNodeContainer& B
 }
 
 TOptional<UE::Interchange::FImportImage> UInterchangeMaterialXTranslator::GetTexturePayloadData(const FString& PayloadKey, TOptional<FString>& AlternateTexturePath) const
+{
+	return GetTexturePayloadData(PayloadKey, AlternateTexturePath, Results, AnalyticsHelper);
+}
+
+TOptional<UE::Interchange::FImportImage> UInterchangeMaterialXTranslator::GetTexturePayloadData(const FString& PayloadKey, TOptional<FString>& AlternateTexturePath, UInterchangeResultsContainer* InResults, TSharedPtr<UE::Interchange::FAnalyticsHelper> InAnalyticsHelper)
 {
 	FString Filename = PayloadKey;
 	TextureCompressionSettings CompressionSettings = TextureCompressionSettings::TC_Default;
 
 #if WITH_EDITOR
-	if(int32 IndexTextureCompression; PayloadKey.FindChar(FMaterialXManager::TexturePayloadSeparator, IndexTextureCompression))
+	if (int32 IndexTextureCompression; PayloadKey.FindChar(FMaterialXManager::TexturePayloadSeparator, IndexTextureCompression))
 	{
 		Filename = PayloadKey.Mid(0, IndexTextureCompression);
 		CompressionSettings = TextureCompressionSettings(FCString::Atoi(*PayloadKey.Mid(IndexTextureCompression + 1)));
 	}
 #endif
 
-	UE::Interchange::Private::FScopedTranslator ScopedTranslator(Filename, Results, AnalyticsHelper);
+	UE::Interchange::Private::FScopedTranslator ScopedTranslator(Filename, InResults, InAnalyticsHelper);
 	const IInterchangeTexturePayloadInterface* TextureTranslator = ScopedTranslator.GetPayLoadInterface<IInterchangeTexturePayloadInterface>();
 
-	if(!ensure(TextureTranslator))
+	if (!ensure(TextureTranslator))
 	{
 		return TOptional<UE::Interchange::FImportImage>();
 	}
@@ -84,10 +89,78 @@ TOptional<UE::Interchange::FImportImage> UInterchangeMaterialXTranslator::GetTex
 
 	TOptional<UE::Interchange::FImportImage> TexturePayloadData = TextureTranslator->GetTexturePayloadData(PayloadKey, AlternateTexturePath);
 
-	if(TexturePayloadData.IsSet())
+	if (TexturePayloadData.IsSet())
 	{
 		TexturePayloadData.GetValue().CompressionSettings = CompressionSettings;
 	}
 
 	return TexturePayloadData;
 }
+
+TOptional<UE::Interchange::FImportBlockedImage> UInterchangeMaterialXTranslator::GetBlockedTexturePayloadData(const FString& PayloadKey, TOptional<FString>& AlternateTexturePath) const
+{
+	return GetBlockedTexturePayloadData(PayloadKey, AlternateTexturePath, Results, AnalyticsHelper);
+}
+
+TOptional<UE::Interchange::FImportBlockedImage> UInterchangeMaterialXTranslator::GetBlockedTexturePayloadData(const FString& PayloadKey, TOptional<FString>& AlternateTexturePath, UInterchangeResultsContainer* InResults, TSharedPtr<UE::Interchange::FAnalyticsHelper> InAnalyticsHelper)
+{
+	using namespace UE::Interchange;
+	using namespace UE::TextureUtilitiesCommon;
+
+	FString Filename = PayloadKey;
+	TextureCompressionSettings CompressionSettings = TextureCompressionSettings::TC_Default;
+
+#if WITH_EDITOR
+	if (int32 IndexTextureCompression; PayloadKey.FindChar(FMaterialXManager::TexturePayloadSeparator, IndexTextureCompression))
+	{
+		Filename = PayloadKey.Mid(0, IndexTextureCompression);
+		CompressionSettings = TextureCompressionSettings(FCString::Atoi(*PayloadKey.Mid(IndexTextureCompression + 1)));
+	}
+#endif
+
+	FImportBlockedImage BlockedImage;
+	BlockedImage.CompressionSettings = CompressionSettings;
+
+	TMap<int32, FString> UDIMBlocks = GetUDIMBlocksFromSourceFile(Filename, DefaultUdimRegexPattern);
+
+	TArray<FImportImage> Images;
+	Images.Reserve(UDIMBlocks.Num());
+
+	for (const TPair<int32, FString>& Pair : UDIMBlocks)
+	{
+		const int32 IndexUDIM = Pair.Get<0>();
+		const FString& FilenameUDIM = Pair.Get<1>();
+
+		UE::Interchange::Private::FScopedTranslator ScopedTranslator(FilenameUDIM, InResults, InAnalyticsHelper);
+		const IInterchangeTexturePayloadInterface* TextureTranslator = ScopedTranslator.GetPayLoadInterface<IInterchangeTexturePayloadInterface>();
+
+		// should we stop if one of the UDIM image is broken?
+		if (!ensure(TextureTranslator))
+		{
+			continue;
+		}
+
+		AlternateTexturePath = FilenameUDIM;
+		TOptional<UE::Interchange::FImportImage> TexturePayloadData = TextureTranslator->GetTexturePayloadData(FilenameUDIM, AlternateTexturePath);
+
+		if (TexturePayloadData.IsSet())
+		{
+			const FImportImage& Image = Images.Emplace_GetRef(MoveTemp(*TexturePayloadData));
+			BlockedImage.Format = Image.Format;
+
+			int32 BlockX = INDEX_NONE, BlockY = INDEX_NONE;
+			ExtractUDIMCoordinates(IndexUDIM, BlockX, BlockY);
+
+			if (BlockX == INDEX_NONE || BlockY == INDEX_NONE)
+			{
+				continue;
+			}
+
+			BlockedImage.InitBlockFromImage(BlockX, BlockY, Image);
+		}
+	}
+
+	BlockedImage.MigrateDataFromImagesToRawData(Images);
+
+	return BlockedImage;
+}
diff --git a/Engine/Plugins/Interchange/Runtime/Source/Import/Public/MaterialX/InterchangeMaterialXTranslator.h b/Engine/Plugins/Interchange/Runtime/Source/Import/Public/MaterialX/InterchangeMaterialXTranslator.h
--- a/Engine/Plugins/Interchange/Runtime/Source/Import/Public/MaterialX/InterchangeMaterialXTranslator.h
+++ b/Engine/Plugins/Interchange/Runtime/Source/Import/Public/MaterialX/InterchangeMaterialXTranslator.h
@@ -1,10 +1,9 @@
 // Copyright Epic Games, Inc. All Rights Reserved.
 #pragma once
 
-#include "InterchangeShaderGraphNode.h"
 #include "InterchangeTranslatorBase.h"
-#include "InterchangeTextureNode.h"
 #include "Texture/InterchangeTexturePayloadInterface.h"
+#include "Texture/InterchangeBlockedTexturePayloadInterface.h"
 
 #include "InterchangeMaterialXTranslator.generated.h"
 
@@ -15,7 +14,10 @@ class UInterchangeBaseLightNode;
 class UInterchangeSceneNode;
 
 UCLASS(MinimalAPI, BlueprintType)
-class UInterchangeMaterialXTranslator : public UInterchangeTranslatorBase, public IInterchangeTexturePayloadInterface
+class UInterchangeMaterialXTranslator
+	: public UInterchangeTranslatorBase
+	, public IInterchangeTexturePayloadInterface
+	, public IInterchangeBlockedTexturePayloadInterface
 {
 	GENERATED_BODY()
 
@@ -38,7 +40,14 @@ public:
 	UE_API virtual bool Translate( UInterchangeBaseNodeContainer& BaseNodeContainer ) const override;
 	/** End UInterchangeTranslatorBase API*/
 
+	using IInterchangeTexturePayloadInterface::GetTexturePayloadData;
 	UE_API virtual TOptional<UE::Interchange::FImportImage> GetTexturePayloadData(const FString& PayloadKey, TOptional<FString>& AlternateTexturePath) const override;
+	UE_API static TOptional<UE::Interchange::FImportImage> GetTexturePayloadData(const FString& PayloadKey, TOptional<FString>& AlternateTexturePath, UInterchangeResultsContainer* InResults, TSharedPtr<UE::Interchange::FAnalyticsHelper> InAnalyticsHelper);
+
+	using IInterchangeBlockedTexturePayloadInterface::GetBlockedTexturePayloadData;
+	UE_API virtual TOptional<UE::Interchange::FImportBlockedImage> GetBlockedTexturePayloadData(const FString& PayloadKey, TOptional<FString>& AlternateTexturePath) const override;
+	UE_API static TOptional<UE::Interchange::FImportBlockedImage> GetBlockedTexturePayloadData(const FString& PayloadKey, TOptional<FString>& AlternateTexturePath, UInterchangeResultsContainer* InResults, TSharedPtr<UE::Interchange::FAnalyticsHelper> InAnalyticsHelper);
+
 };
 
 #undef UE_API
diff --git a/Engine/Plugins/Interchange/Runtime/Source/Import/Public/MaterialX/MaterialXUtils/MaterialXSurfaceShaderAbstract.h b/Engine/Plugins/Interchange/Runtime/Source/Import/Public/MaterialX/MaterialXUtils/MaterialXSurfaceShaderAbstract.h
--- a/Engine/Plugins/Interchange/Runtime/Source/Import/Public/MaterialX/MaterialXUtils/MaterialXSurfaceShaderAbstract.h
+++ b/Engine/Plugins/Interchange/Runtime/Source/Import/Public/MaterialX/MaterialXUtils/MaterialXSurfaceShaderAbstract.h
@@ -403,6 +403,7 @@ protected:
 				else
 				{
 					Filepath = InputFile->hasValue() ? InputFile->getValueString().c_str() : InputFile->getInterfaceInput()->getValueString().c_str();
+					Filepath.ReplaceInline(TEXT(".<UDIM>."), TEXT(".1001."));
 					const FString FilePrefix = GetFilePrefix(InputFile);
 					Filepath = FPaths::Combine(FilePrefix, Filepath);
 				}
