USD: Add support for dome light as sky light components when importing via USD Interchange.

diff --git a/Engine/Plugins/Interchange/Extensions/OpenUSD/Source/Import/Private/InterchangeUsdTranslator.cpp b/Engine/Plugins/Interchange/Extensions/OpenUSD/Source/Import/Private/InterchangeUsdTranslator.cpp
--- a/Engine/Plugins/Interchange/Extensions/OpenUSD/Source/Import/Private/InterchangeUsdTranslator.cpp
+++ b/Engine/Plugins/Interchange/Extensions/OpenUSD/Source/Import/Private/InterchangeUsdTranslator.cpp
@@ -34,6 +34,7 @@
 #include "UsdWrappers/UsdStage.h"
 #include "UsdWrappers/UsdTyped.h"
 #include "UsdWrappers/UsdVariantSets.h"
+#include "UsdWrappers/VtValue.h"
 
 #include "InterchangeCameraNode.h"
 #include "InterchangeLightNode.h"
@@ -45,6 +46,7 @@
 #include "InterchangeSceneNode.h"
 #include "InterchangeShaderGraphNode.h"
 #include "InterchangeTexture2DNode.h"
+#include "InterchangeTextureCubeNode.h"
 #include "InterchangeTranslatorHelper.h"
 #include "InterchangeUsdContext.h"
 #include "InterchangeVolumeNode.h"
@@ -105,13 +107,14 @@ namespace UE::InterchangeUsdTranslator::Private
 	const static FString BonePrefix = TEXT("\\Bone\\");
 	const static FString CameraPrefix = TEXT("\\Camera\\");
 	const static FString LightPrefix = TEXT("\\Light\\");
-	const static FString VolumePrefix = TEXT("\\Volume\\");
 	const static FString LODPrefix = TEXT("\\LOD\\");
 	const static FString MaterialPrefix = TEXT("\\Material\\");
 	const static FString MaterialReferencePrefix = TEXT("\\MaterialReference\\");
 	const static FString MeshPrefix = TEXT("\\Mesh\\");
 	const static FString MorphTargetPrefix = TEXT("\\MorphTarget\\");
 	const static FString PrimitiveShapePrefix = TEXT("\\PrimitiveShape\\");
+	const static FString TexturePrefix = TEXT("\\Texture\\");
+	const static FString VolumePrefix = TEXT("\\Volume\\");
 
 	const static FString LODContainerSuffix = TEXT("LODContainer");
 	const static FString TwoSidedSuffix = TEXT("_TwoSided");
@@ -924,6 +927,34 @@ namespace UE::InterchangeUsdTranslator::Private
 		}
 	}
 
+	FString AddTextureCubeNode(const FString& TextureFilePath, UInterchangeBaseNodeContainer& NodeContainer)
+	{
+		const FString NodeUid = FString::Printf(TEXT("%s%s"), *TexturePrefix, *TextureFilePath);
+
+		// If we created an actual TextureCubeNode here we'd be queried for the sliced payload later. Our HDR image is
+		// just a latlong 2d texture though, so we instead need to create a Texture2DNode and then set SetForceLongLatCubemap(),
+		// which will cause UInterchangeGenericTexturePipeline::HandleCreationOfTextureFactoryNode() to produce a
+		// UInterchangeTextureCubeFactoryNode from this translated node, producing our TextureCube
+		if (const UInterchangeTexture2DNode* Node = Cast<UInterchangeTexture2DNode>(NodeContainer.GetNode(NodeUid)))
+		{
+			return NodeUid;
+		}
+
+		FString NodeName{FPaths::GetCleanFilename(TextureFilePath)};
+
+		UInterchangeTexture2DNode* Node = NewObject<UInterchangeTexture2DNode>(&NodeContainer);
+		NodeContainer.SetupNode(Node, NodeUid, NodeName, EInterchangeNodeContainerType::TranslatedAsset);
+
+		Node->SetForceLongLatCubemap(true);
+
+		UsdToUnreal::FTextureParameterValue TextureParameter;
+		TextureParameter.TextureFilePath = TextureFilePath;
+		TextureParameter.Group = TextureGroup::TEXTUREGROUP_Skybox;
+		Node->SetPayLoadKey(EncodeTexturePayloadKey(TextureParameter));
+
+		return NodeUid;
+	}
+
 	FString AddDisplayColorMaterialInstanceNodeIfNeeded(UInterchangeBaseNodeContainer& NodeContainer, const UsdUtils::FUsdPrimMaterialSlot& Slot)
 	{
 		using namespace UsdUnreal::MaterialUtils;
@@ -1243,7 +1274,8 @@ namespace UE::InterchangeUsdTranslator::Private
 			// Emit texture node itself (this is the main place where this happens)
 			// Note that the node name isn't just the texture path, as we may have multiple material users of this texture
 			// with different settings, and so we need separate translated nodes for each material and parameter
-			const FString TextureUid = FString::Printf(TEXT("Texture:%s:%s"), *UsdUtils::GetPrototypePrimPath(Prim).GetString(), **BaseParameterName);
+			// TODO: Encode the different settings on the node uid as well to automatically deduplicate?
+			const FString TextureUid = FString::Printf(TEXT("%s%s:%s"), *TexturePrefix, *UsdUtils::GetPrototypePrimPath(Prim).GetString(), **BaseParameterName);
 			AddTextureNode(Prim, TextureUid, Value, NodeContainer);
 
 			// Actual texture assignment
@@ -1825,7 +1857,43 @@ namespace UE::InterchangeUsdTranslator::Private
 
 		float Intensity = UsdUtils::GetAttributeValue<float>(Prim, IntensityToken);
 		float Exposure = UsdUtils::GetAttributeValue<float>(Prim, ExposureToken);
-		FLinearColor Color = UsdUtils::GetAttributeValue<FLinearColor>(Prim, ColorToken);
+
+		// Temporarily a bit awkward here as we updated our VtValue::Get() to not perform any type coercion, so we
+		// need to do that ourselves. In the future we'll have a dedicated templated function to perform this process
+		FLinearColor Color;
+		UE::FUsdAttribute ColorAttribute = Prim.GetAttribute(*ColorToken);
+		if (ColorAttribute && ColorAttribute.HasAuthoredValue())
+		{
+			UE::FVtValue Value;
+			if (ColorAttribute.Get(Value))
+			{
+				const FString TypeName = Value.GetTypeName();
+				if (TypeName == TEXT("GfVec3f"))
+				{
+					Color = FLinearColor{Value.Get<FVector3f>()};
+				}
+				else if (TypeName == TEXT("GfVec3h"))
+				{
+					Color = FLinearColor{Value.Get<FVector3h>()};
+				}
+				else if (TypeName == TEXT("GfVec3d"))
+				{
+					Color = FLinearColor{Value.Get<FVector3d>()};
+				}
+				if (TypeName == TEXT("GfVec4f"))
+				{
+					Color = FLinearColor{Value.Get<FVector4f>()};
+				}
+				else if (TypeName == TEXT("GfVec4h"))
+				{
+					Color = FLinearColor{Value.Get<FVector4h>()};
+				}
+				else if (TypeName == TEXT("GfVec4d"))
+				{
+					Color = FLinearColor{Value.Get<FVector4d>()};
+				}
+			}
+		}
 
 		const bool bSRGB = true;
 		Color = Color.ToFColor(bSRGB);
@@ -1843,6 +1911,7 @@ namespace UE::InterchangeUsdTranslator::Private
 			NodeContainer.SetupNode(LightNode, NodeUid, NodeName, EInterchangeNodeContainerType::TranslatedAsset);
 			LightNode->SetAssetName(NodeName);
 
+			// TODO: Maybe only authore these in case there were in fact authored opinions on the prims? Otherwise we're just authoring default values
 			LightNode->SetCustomLightColor(Color);
 			LightNode->SetCustomTemperature(Temperature);
 			LightNode->SetCustomUseTemperature(UseTemperature);
@@ -1934,8 +2003,48 @@ namespace UE::InterchangeUsdTranslator::Private
 			LightNode->SetCustomSourceWidth(Width);
 			LightNode->SetCustomSourceHeight(Height);
 		}
-		// #ueent_todo:
-		// DomeLight -> SkyLight
+		else if (Prim.IsA(TEXT("DomeLight")))
+		{
+			UInterchangeSkyLightNode* LightNode = NewObject<UInterchangeSkyLightNode>(&NodeContainer);
+			SetBaseLightProperties(LightNode);
+
+			// This doesn't translate very well as SkyLightComponents seem to have an "IntensityScale", and there's never any clear
+			// indication of what the dome light intensity actually is... But I don't think there's much we can do
+			Intensity = UsdToUnreal::ConvertLightIntensityAttr(Intensity, Exposure);
+			LightNode->SetCustomIntensity(Intensity);
+
+			// Show some warning in case the file specifies a different DomeLight format, as we only support latlong
+			static const FString FormatToken = UsdToUnreal::ConvertToken(pxr::UsdLuxTokens->inputsTextureFormat);
+			UE::FUsdAttribute FormatAttribute = Prim.GetAttribute(*FormatToken);
+			FString Format;
+			if (FormatAttribute.HasAuthoredValue() &&
+				FormatAttribute.Get<FString>(Format) &&
+				!Format.IsEmpty() &&
+				Format != TEXT("automatic") &&
+				Format != TEXT("latlong"))
+			{
+				UE_LOG(
+					LogUsd,
+					Warning,
+					TEXT(
+						"DomeLight '%s' specifies format '%s', but only 'latlong' DomeLights are supported for now. The texture will be treated as 'latlong'"
+					),
+					*Prim.GetPrimPath().GetString(),
+					*Format
+				);
+			}
+
+			static const FString FileToken = UsdToUnreal::ConvertToken(pxr::UsdLuxTokens->inputsTextureFile);
+			UE::FUsdAttribute FileAttribute = Prim.GetAttribute(*FileToken);
+			const FString ResolvedDomeTexturePath = UsdUtils::GetResolvedAssetPath(FileAttribute);
+
+			FString TextureCubeNodeUid = AddTextureCubeNode(ResolvedDomeTexturePath, NodeContainer);
+			if (!TextureCubeNodeUid.IsEmpty())
+			{
+				LightNode->SetCustomCubemapDependency(TextureCubeNodeUid);
+				LightNode->SetCustomSourceType(EInterchangeSkyLightSourceType::SpecifiedCubemap);
+			}
+		}
 
 		return NodeUid;
 	}
@@ -2357,7 +2466,7 @@ namespace UE::InterchangeUsdTranslator::Private
 			else
 			{
 				UE_LOG(
-					LogTemp,
+					LogUsd,
 					Warning,
 					TEXT(
 						"Failed to parse a LOD Mesh from variant '%s' of prim '%s'. For automatic parsing of LODs, make sure there is a single Mesh prim within the variant, named exactly as the variant itself (e.g. 'LOD0', 'LOD1', etc.)"
@@ -4856,6 +4965,10 @@ TOptional<UE::Interchange::FImportImage> UInterchangeUSDTranslator::GetTexturePa
 	if (bDecoded)
 	{
 		// Defer back to another translator to actually parse the texture raw data
+		//
+		// Note: For DomeLights even though we produce texture cube nodes, we'll rely on a common texture payload here (i.e. not sliced),
+		// as Unreal handles only latlong / equirectangular, single slice .HDRs anyway. Via the .dds format you can provide sliced payloads
+		// for texture cubes, but we do not handle that via USD for now
 		UE::Interchange::Private::FScopedTranslator ScopedTranslator(FilePath, Results, AnalyticsHelper);
 		const IInterchangeTexturePayloadInterface* TextureTranslator = ScopedTranslator.GetPayLoadInterface<IInterchangeTexturePayloadInterface>();
 		if (TextureTranslator)
@@ -4864,12 +4977,13 @@ TOptional<UE::Interchange::FImportImage> UInterchangeUSDTranslator::GetTexturePa
 			const FString UnusedPayloadKey = {};
 			TexturePayloadData = TextureTranslator->GetTexturePayloadData(UnusedPayloadKey, AlternateTexturePath);
 
-			// Move compression settings onto the payload data.
-			// Note: We don't author anything else on the texture payload data here (like the sRGB flag), because those
-			// settings were already on our translated node, and presumably already made their way to the factory node.
-			// The factory should use them to override whatever it finds in this payload data, with the exception of the
-			// compression settings (which can't be stored on the translated node)
-			TexturePayloadData->CompressionSettings = TextureGroup == TEXTUREGROUP_WorldNormalMap ? TC_Normalmap : TC_Default;
+			// The TextureTranslator doesn't know how we plan on using this texture: If we know it's meant to be a normal map,
+			// then override the compression settings here. Leave the rest untouched as the TextureTranslator is going to set
+			// that according to the file itself
+			if (TextureGroup == TEXTUREGROUP_WorldNormalMap)
+			{
+				TexturePayloadData->CompressionSettings = TC_Normalmap;
+			}
 		}
 	}
 
diff --git a/Engine/Plugins/Interchange/Runtime/Source/FactoryNodes/Private/InterchangeLightFactoryNode.cpp b/Engine/Plugins/Interchange/Runtime/Source/FactoryNodes/Private/InterchangeLightFactoryNode.cpp
--- a/Engine/Plugins/Interchange/Runtime/Source/FactoryNodes/Private/InterchangeLightFactoryNode.cpp
+++ b/Engine/Plugins/Interchange/Runtime/Source/FactoryNodes/Private/InterchangeLightFactoryNode.cpp
@@ -3,6 +3,7 @@
 #include "InterchangeLightFactoryNode.h"
 
 #include "Components/RectLightComponent.h"
+#include "Components/SkyLightComponent.h"
 #include "Components/SpotLightComponent.h"
 
 #include UE_INLINE_GENERATED_CPP_BY_NAME(InterchangeLightFactoryNode)
@@ -14,7 +15,7 @@ bool UInterchangeBaseLightFactoryNode::GetCustomLightColor(FColor& AttributeValu
 
 bool UInterchangeBaseLightFactoryNode::SetCustomLightColor(const FColor& AttributeValue, bool bAddApplyDelegate)
 {
-	IMPLEMENT_NODE_ATTRIBUTE_SETTER(UInterchangeBaseLightFactoryNode, LightColor, FColor, ULightComponent)
+	IMPLEMENT_NODE_ATTRIBUTE_SETTER(UInterchangeBaseLightFactoryNode, LightColor, FColor, ULightComponentBase)
 }
 
 bool UInterchangeBaseLightFactoryNode::GetCustomIntensity(float& AttributeValue) const
@@ -24,7 +25,7 @@ bool UInterchangeBaseLightFactoryNode::GetCustomIntensity(float& AttributeValue)
 
 bool UInterchangeBaseLightFactoryNode::SetCustomIntensity(float AttributeValue, bool bAddApplyDelegate)
 {
-	IMPLEMENT_NODE_ATTRIBUTE_SETTER(UInterchangeBaseLightFactoryNode, Intensity, float, ULightComponent)
+	IMPLEMENT_NODE_ATTRIBUTE_SETTER(UInterchangeBaseLightFactoryNode, Intensity, float, ULightComponentBase)
 }
 
 bool UInterchangeBaseLightFactoryNode::GetCustomTemperature(float& AttributeValue) const
@@ -53,8 +54,8 @@ void UInterchangeBaseLightFactoryNode::CopyWithObject(const UInterchangeFactoryB
 
 	if (const UInterchangeBaseLightFactoryNode* LightFactoryNode = Cast<UInterchangeBaseLightFactoryNode>(SourceNode))
 	{
-		COPY_NODE_DELEGATES(LightFactoryNode, LightColor, FColor, ULightComponent)
-		COPY_NODE_DELEGATES(LightFactoryNode, Intensity, float, ULightComponent)
+		COPY_NODE_DELEGATES(LightFactoryNode, LightColor, FColor, ULightComponentBase)
+		COPY_NODE_DELEGATES(LightFactoryNode, Intensity, float, ULightComponentBase)
 		COPY_NODE_DELEGATES(LightFactoryNode, Temperature, float, ULightComponent)
 		COPY_NODE_DELEGATES(LightFactoryNode, bUseTemperature, bool, ULightComponent)
 	}
@@ -248,3 +249,23 @@ void UInterchangeRectLightFactoryNode::OnRestoreAllCustomAttributeDelegates()
 	REFILL_CUSTOM_ATTRIBUTE_APPLY_DELEGATE(SourceWidth, float)
 	REFILL_CUSTOM_ATTRIBUTE_APPLY_DELEGATE(SourceHeight, float)
 }
+
+bool UInterchangeSkyLightFactoryNode::GetCustomCubemapDependency(FString& AttributeValue) const
+{
+	IMPLEMENT_NODE_ATTRIBUTE_GETTER(CubemapDependency, FString)
+}
+
+bool UInterchangeSkyLightFactoryNode::SetCustomCubemapDependency(FString AttributeValue)
+{
+	IMPLEMENT_NODE_ATTRIBUTE_SETTER_NODELEGATE(CubemapDependency, FString)
+}
+
+bool UInterchangeSkyLightFactoryNode::GetCustomSourceType(EInterchangeSkyLightSourceType& AttributeValue) const
+{
+	IMPLEMENT_NODE_ATTRIBUTE_GETTER(SourceType, EInterchangeSkyLightSourceType)
+}
+
+bool UInterchangeSkyLightFactoryNode::SetCustomSourceType(EInterchangeSkyLightSourceType AttributeValue)
+{
+	IMPLEMENT_NODE_ATTRIBUTE_SETTER_NODELEGATE(SourceType, EInterchangeSkyLightSourceType)
+}
diff --git a/Engine/Plugins/Interchange/Runtime/Source/FactoryNodes/Public/InterchangeLightFactoryNode.h b/Engine/Plugins/Interchange/Runtime/Source/FactoryNodes/Public/InterchangeLightFactoryNode.h
--- a/Engine/Plugins/Interchange/Runtime/Source/FactoryNodes/Public/InterchangeLightFactoryNode.h
+++ b/Engine/Plugins/Interchange/Runtime/Source/FactoryNodes/Public/InterchangeLightFactoryNode.h
@@ -4,6 +4,7 @@
 
 #include "CoreMinimal.h"
 #include "InterchangeActorFactoryNode.h"
+#include "InterchangeLightNode.h"
 
 #include "InterchangeLightFactoryNode.generated.h"
 
@@ -193,4 +194,29 @@ private:
 	IMPLEMENT_NODE_ATTRIBUTE_KEY(OuterConeAngle)
 };
 
+UCLASS(MinimalAPI, BlueprintType)
+class UInterchangeSkyLightFactoryNode : public UInterchangeBaseLightFactoryNode
+{
+	GENERATED_BODY()
+
+public:
+
+	UFUNCTION(BlueprintCallable, Category = "Interchange | Node | SkyLightFactory")
+	UE_API bool GetCustomCubemapDependency(FString& TextureCubeFactoryNodeUid) const;
+
+	UFUNCTION(BlueprintCallable, Category = "Interchange | Node | SkyLightFactory")
+	UE_API bool SetCustomCubemapDependency(FString TextureCubeFactoryNodeUid);
+
+	UFUNCTION(BlueprintCallable, Category = "Interchange | Node | SkyLightFactory")
+	UE_API bool GetCustomSourceType(EInterchangeSkyLightSourceType& SourceType) const;
+
+	UFUNCTION(BlueprintCallable, Category = "Interchange | Node | SkyLightFactory")
+	UE_API bool SetCustomSourceType(EInterchangeSkyLightSourceType SourceType);
+
+private:
+
+	IMPLEMENT_NODE_ATTRIBUTE_KEY(CubemapDependency)
+	IMPLEMENT_NODE_ATTRIBUTE_KEY(SourceType)
+};
+
 #undef UE_API
diff --git a/Engine/Plugins/Interchange/Runtime/Source/Import/Private/InterchangeImportModule.cpp b/Engine/Plugins/Interchange/Runtime/Source/Import/Private/InterchangeImportModule.cpp
--- a/Engine/Plugins/Interchange/Runtime/Source/Import/Private/InterchangeImportModule.cpp
+++ b/Engine/Plugins/Interchange/Runtime/Source/Import/Private/InterchangeImportModule.cpp
@@ -31,6 +31,7 @@
 #include "Scene/InterchangeSceneImportAssetFactory.h"
 #include "Scene/InterchangeSceneVariantSetsFactory.h"
 #include "Scene/InterchangeSkeletalMeshActorFactory.h"
+#include "Scene/InterchangeSkyLightActorFactory.h"
 #include "Scene/InterchangeStaticMeshActorFactory.h"
 #include "SpecularProfile/InterchangeSpecularProfileFactory.h"
 #include "Texture/InterchangeDDSTranslator.h"
@@ -121,6 +122,7 @@ void FInterchangeImportModule::StartupModule()
 		InterchangeManager.RegisterFactory(UInterchangeSkeletalMeshActorFactory::StaticClass());
 		InterchangeManager.RegisterFactory(UInterchangeSceneVariantSetsFactory::StaticClass());
 		InterchangeManager.RegisterFactory(UInterchangeLightActorFactory::StaticClass());
+		InterchangeManager.RegisterFactory(UInterchangeSkyLightActorFactory::StaticClass());
 		InterchangeManager.RegisterFactory(UInterchangeSceneImportAssetFactory::StaticClass());
 		InterchangeManager.RegisterFactory(UInterchangeSpecularProfileFactory::StaticClass());
 		InterchangeManager.RegisterFactory(UInterchangeHeterogeneousVolumeActorFactory::StaticClass());
diff --git a/Engine/Plugins/Interchange/Runtime/Source/Import/Private/Scene/InterchangeSkyLightActorFactory.cpp b/Engine/Plugins/Interchange/Runtime/Source/Import/Private/Scene/InterchangeSkyLightActorFactory.cpp
new file mode 100644
--- /dev/null
+++ b/Engine/Plugins/Interchange/Runtime/Source/Import/Private/Scene/InterchangeSkyLightActorFactory.cpp
@@ -0,0 +1,92 @@
+// Copyright Epic Games, Inc. All Rights Reserved.
+
+#include "Scene/InterchangeSkyLightActorFactory.h"
+
+#include "InterchangeTextureCubeFactoryNode.h"
+#include "Nodes/InterchangeBaseNodeContainer.h"
+#include "Nodes/InterchangeFactoryBaseNode.h"
+#include "Scene/InterchangeActorHelper.h"
+
+#include "Components/SkyLightComponent.h"
+#include "CoreMinimal.h"
+#include "Engine/SkyLight.h"
+#include "InterchangeLightFactoryNode.h"
+#include "UObject/Object.h"
+#include "UObject/ObjectMacros.h"
+
+#include UE_INLINE_GENERATED_CPP_BY_NAME(InterchangeSkyLightActorFactory)
+
+UClass* UInterchangeSkyLightActorFactory::GetFactoryClass() const
+{
+	return ASkyLight::StaticClass();
+}
+
+void UInterchangeSkyLightActorFactory::SetupObject_GameThread(const FSetupObjectParams& Arguments)
+{
+	Super::SetupObject_GameThread(Arguments);
+
+	if (!ensure(Arguments.NodeContainer))
+	{
+		return;
+	}
+
+	if (ASkyLight* SkyLightActor = Cast<ASkyLight>(Arguments.ImportedObject))
+	{
+		if (USkyLightComponent* SkyLightComponent = SkyLightActor->GetLightComponent())
+		{
+			SkyLightComponent->UnregisterComponent();
+
+			if (const UInterchangeSkyLightFactoryNode* SkyLightFactoryNode = Cast<UInterchangeSkyLightFactoryNode>(Arguments.FactoryNode))
+			{
+				EInterchangeSkyLightSourceType SourceType;
+				if (SkyLightFactoryNode->GetCustomSourceType(SourceType))
+				{
+					static_assert((int)EInterchangeSkyLightSourceType::CapturedScene == (int)ESkyLightSourceType::SLS_CapturedScene);
+					static_assert((int)EInterchangeSkyLightSourceType::SpecifiedCubemap == (int)ESkyLightSourceType::SLS_SpecifiedCubemap);
+					SkyLightComponent->SourceType = static_cast<ESkyLightSourceType>(SourceType);
+				}
+
+				FString TextureFactoryNodeUid;
+				if (SkyLightFactoryNode->GetCustomCubemapDependency(TextureFactoryNodeUid) && !TextureFactoryNodeUid.IsEmpty())
+				{
+					UInterchangeTextureCubeFactoryNode* TextureCubeFactoryNode = Cast<UInterchangeTextureCubeFactoryNode>(
+						Arguments.NodeContainer->GetFactoryNode(TextureFactoryNodeUid)
+					);
+
+					if (TextureCubeFactoryNode)
+					{
+						FSoftObjectPath ReferenceObject;
+						TextureCubeFactoryNode->GetCustomReferenceObject(ReferenceObject);
+						if (UTextureCube* Texture = Cast<UTextureCube>(ReferenceObject.TryLoad()))
+						{
+							// We must assign the cubemap from SetupObject_GameThread instead of ProcessActor, as this call
+							// will invalidate the reflection capture and get FScene::CaptureOrUploadReflectionCapture() to try and
+							// read the UTextureCube::GetResource() immediately (or on the next frame?). By the time ProcessActor
+							// runs however, we haven't run UpdateResource() on the UTextureCube yet, so this process would have
+							// failed and produced a warning
+							SkyLightComponent->SetCubemap(Texture);
+						}
+					}
+				}
+			}
+		}
+	}
+}
+
+UObject* UInterchangeSkyLightActorFactory::ProcessActor(
+	AActor& SpawnedActor,
+	const UInterchangeActorFactoryNode& FactoryNode,
+	const UInterchangeBaseNodeContainer& NodeContainer,
+	const FImportSceneObjectsParams& Params
+)
+{
+	if (ASkyLight* SkyLightActor = Cast<ASkyLight>(&SpawnedActor))
+	{
+		if (USkyLightComponent* SkyLightComponent = SkyLightActor->GetLightComponent())
+		{
+			return SkyLightComponent;
+		}
+	}
+
+	return nullptr;
+}
diff --git a/Engine/Plugins/Interchange/Runtime/Source/Import/Public/Scene/InterchangeSkyLightActorFactory.h b/Engine/Plugins/Interchange/Runtime/Source/Import/Public/Scene/InterchangeSkyLightActorFactory.h
new file mode 100644
--- /dev/null
+++ b/Engine/Plugins/Interchange/Runtime/Source/Import/Public/Scene/InterchangeSkyLightActorFactory.h
@@ -0,0 +1,29 @@
+// Copyright Epic Games, Inc. All Rights Reserved.
+
+#pragma once
+
+#include "Scene/InterchangeActorFactory.h"
+
+#include "InterchangeSkyLightActorFactory.generated.h"
+
+#define UE_API INTERCHANGEIMPORT_API
+
+UCLASS(MinimalAPI, BlueprintType)
+class UInterchangeSkyLightActorFactory : public UInterchangeActorFactory
+{
+	GENERATED_BODY()
+
+public:
+	UE_API virtual UClass* GetFactoryClass() const override;
+	UE_API virtual void SetupObject_GameThread(const FSetupObjectParams& Arguments) override;
+
+protected:
+	UE_API virtual UObject* ProcessActor(
+		class AActor& SpawnedActor,
+		const UInterchangeActorFactoryNode& FactoryNode,
+		const UInterchangeBaseNodeContainer& NodeContainer,
+		const FImportSceneObjectsParams& Params
+	) override;
+};
+
+#undef UE_API
diff --git a/Engine/Plugins/Interchange/Runtime/Source/Nodes/Private/InterchangeLightNode.cpp b/Engine/Plugins/Interchange/Runtime/Source/Nodes/Private/InterchangeLightNode.cpp
--- a/Engine/Plugins/Interchange/Runtime/Source/Nodes/Private/InterchangeLightNode.cpp
+++ b/Engine/Plugins/Interchange/Runtime/Source/Nodes/Private/InterchangeLightNode.cpp
@@ -199,3 +199,28 @@ FString UInterchangeDirectionalLightNode::GetTypeName() const
 {
 	return TEXT("DirectionalLightNode");
 }
+
+FString UInterchangeSkyLightNode::GetTypeName() const
+{
+	return TEXT("SkyLightNode");
+}
+
+bool UInterchangeSkyLightNode::GetCustomCubemapDependency(FString& AttributeValue) const
+{
+	IMPLEMENT_NODE_ATTRIBUTE_GETTER(CubemapDependency, FString)
+}
+
+bool UInterchangeSkyLightNode::SetCustomCubemapDependency(FString AttributeValue)
+{
+	IMPLEMENT_NODE_ATTRIBUTE_SETTER_NODELEGATE(CubemapDependency, FString)
+}
+
+bool UInterchangeSkyLightNode::GetCustomSourceType(EInterchangeSkyLightSourceType& AttributeValue) const
+{
+	IMPLEMENT_NODE_ATTRIBUTE_GETTER(SourceType, EInterchangeSkyLightSourceType)
+}
+
+bool UInterchangeSkyLightNode::SetCustomSourceType(EInterchangeSkyLightSourceType AttributeValue)
+{
+	IMPLEMENT_NODE_ATTRIBUTE_SETTER_NODELEGATE(SourceType, EInterchangeSkyLightSourceType)
+}
diff --git a/Engine/Plugins/Interchange/Runtime/Source/Nodes/Public/InterchangeLightNode.h b/Engine/Plugins/Interchange/Runtime/Source/Nodes/Public/InterchangeLightNode.h
--- a/Engine/Plugins/Interchange/Runtime/Source/Nodes/Public/InterchangeLightNode.h
+++ b/Engine/Plugins/Interchange/Runtime/Source/Nodes/Public/InterchangeLightNode.h
@@ -19,6 +19,14 @@ enum class EInterchangeLightUnits : uint8
 	EV,
 };
 
+// Mirrors ESkyLightSourceType
+UENUM()
+enum class EInterchangeSkyLightSourceType : uint8
+{
+	CapturedScene,
+	SpecifiedCubemap
+};
+
 UCLASS(MinimalAPI, BlueprintType)
 class UInterchangeBaseLightNode : public UInterchangeBaseNode
 {
@@ -182,16 +190,16 @@ public:
 
 	UE_API virtual FString GetTypeName() const override;
 
-	UFUNCTION(BlueprintCallable, Category = "Interchange | Node | RectLightFactory")
+	UFUNCTION(BlueprintCallable, Category = "Interchange | Node | RectLight")
 	UE_API bool GetCustomSourceWidth(float& AttributeValue) const;
 
-	UFUNCTION(BlueprintCallable, Category = "Interchange | Node | RectLightFactory")
+	UFUNCTION(BlueprintCallable, Category = "Interchange | Node | RectLight")
 	UE_API bool SetCustomSourceWidth(float AttributeValue);
 
-	UFUNCTION(BlueprintCallable, Category = "Interchange | Node | RectLightFactory")
+	UFUNCTION(BlueprintCallable, Category = "Interchange | Node | RectLight")
 	UE_API bool GetCustomSourceHeight(float& AttributeValue) const;
 
-	UFUNCTION(BlueprintCallable, Category = "Interchange | Node | RectLightFactory")
+	UFUNCTION(BlueprintCallable, Category = "Interchange | Node | RectLight")
 	UE_API bool SetCustomSourceHeight(float AttributeValue);
 
 private:
@@ -209,4 +217,31 @@ public:
 	UE_API virtual FString GetTypeName() const override;
 };
 
+UCLASS(MinimalAPI, BlueprintType)
+class UInterchangeSkyLightNode : public UInterchangeBaseLightNode
+{
+	GENERATED_BODY()
+
+public:
+
+	UE_API virtual FString GetTypeName() const override;
+
+	UFUNCTION(BlueprintCallable, Category = "Interchange | Node | SkyLight")
+	UE_API bool GetCustomCubemapDependency(FString& TextureCubeNodeUid) const;
+
+	UFUNCTION(BlueprintCallable, Category = "Interchange | Node | SkyLight")
+	UE_API bool SetCustomCubemapDependency(FString TextureCubeNodeUid);
+
+	UFUNCTION(BlueprintCallable, Category = "Interchange | Node | SkyLight")
+	UE_API bool GetCustomSourceType(EInterchangeSkyLightSourceType& SourceType) const;
+
+	UFUNCTION(BlueprintCallable, Category = "Interchange | Node | SkyLight")
+	UE_API bool SetCustomSourceType(EInterchangeSkyLightSourceType SourceType);
+
+private:
+
+	IMPLEMENT_NODE_ATTRIBUTE_KEY(CubemapDependency)
+	IMPLEMENT_NODE_ATTRIBUTE_KEY(SourceType)
+};
+
 #undef UE_API
diff --git a/Engine/Plugins/Interchange/Runtime/Source/Nodes/Public/InterchangeTexture2DNode.h b/Engine/Plugins/Interchange/Runtime/Source/Nodes/Public/InterchangeTexture2DNode.h
--- a/Engine/Plugins/Interchange/Runtime/Source/Nodes/Public/InterchangeTexture2DNode.h
+++ b/Engine/Plugins/Interchange/Runtime/Source/Nodes/Public/InterchangeTexture2DNode.h
@@ -139,7 +139,20 @@ public:
 		IMPLEMENT_NODE_ATTRIBUTE_GETTER(WrapV, EInterchangeTextureWrapMode);
 	}
 
+	UFUNCTION(BlueprintCallable, Category = "Interchange | Node | Texture")
+	bool SetForceLongLatCubemap(const bool& AttributeValue)
+	{
+		IMPLEMENT_NODE_ATTRIBUTE_SETTER_NODELEGATE(ForceLongLatCubemap, bool);
+	}
+
+	UFUNCTION(BlueprintCallable, Category = "Interchange | Node | Texture")
+	bool GetForceLongLatCubemap(bool& AttributeValue) const
+	{
+		IMPLEMENT_NODE_ATTRIBUTE_GETTER(ForceLongLatCubemap, bool);
+	}
+
 private:
 	IMPLEMENT_NODE_ATTRIBUTE_KEY(WrapU);
 	IMPLEMENT_NODE_ATTRIBUTE_KEY(WrapV);
+	IMPLEMENT_NODE_ATTRIBUTE_KEY(ForceLongLatCubemap);
 };
diff --git a/Engine/Plugins/Interchange/Runtime/Source/Pipelines/Private/InterchangeGenericScenesPipeline.cpp b/Engine/Plugins/Interchange/Runtime/Source/Pipelines/Private/InterchangeGenericScenesPipeline.cpp
--- a/Engine/Plugins/Interchange/Runtime/Source/Pipelines/Private/InterchangeGenericScenesPipeline.cpp
+++ b/Engine/Plugins/Interchange/Runtime/Source/Pipelines/Private/InterchangeGenericScenesPipeline.cpp
@@ -46,6 +46,7 @@
 #include "Engine/Level.h"
 #include "Engine/PointLight.h"
 #include "Engine/RectLight.h"
+#include "Engine/SkyLight.h"
 #include "Engine/SpotLight.h"
 #include "Engine/StaticMeshActor.h"
 #include "Engine/World.h"
@@ -846,6 +847,10 @@ UInterchangeActorFactoryNode* UInterchangeGenericLevelPipeline::CreateActorFacto
 		{
 			return NewObject<UInterchangeDirectionalLightFactoryNode>(BaseNodeContainer, NAME_None);
 		}
+		else if(TranslatedAssetNode->IsA<UInterchangeSkyLightNode>())
+		{
+			return NewObject<UInterchangeSkyLightFactoryNode>(BaseNodeContainer, NAME_None);
+		}
 		else if (TranslatedAssetNode && TranslatedAssetNode->IsA<UInterchangeDecalNode>())
 		{
 			return NewObject<UInterchangeDecalActorFactoryNode>(BaseNodeContainer, NAME_None);
@@ -1013,6 +1018,22 @@ void UInterchangeGenericLevelPipeline::SetUpFactoryNode(UInterchangeActorFactory
 					}
 				}
 			}
+
+			// Sky Light
+			if (const UInterchangeSkyLightNode* SkyLightNode = Cast<UInterchangeSkyLightNode>(BaseLightNode))
+			{
+				if (UInterchangeSkyLightFactoryNode* SkyLightFactoryNode = Cast<UInterchangeSkyLightFactoryNode>(BaseLightFactoryNode))
+				{
+					if (FString TranslatedTextureUid; SkyLightNode->GetCustomCubemapDependency(TranslatedTextureUid) && !TranslatedTextureUid.IsEmpty())
+					{
+						const FString FactoryTextureUid = UInterchangeFactoryBaseNode::BuildFactoryNodeUid(TranslatedTextureUid);
+						SkyLightFactoryNode->AddFactoryDependencyUid(FactoryTextureUid);
+						SkyLightFactoryNode->SetCustomCubemapDependency(FactoryTextureUid);
+					}
+
+					COPY_FROM_TRANSLATED_TO_FACTORY(SkyLightNode, SkyLightFactoryNode, SourceType, EInterchangeSkyLightSourceType)
+				}
+			}
 		}
 
 		//Test for spot before point since a spot light is a point light
@@ -1032,6 +1053,10 @@ void UInterchangeGenericLevelPipeline::SetUpFactoryNode(UInterchangeActorFactory
 		{
 			ActorFactoryNode->SetCustomActorClassName(ADirectionalLight::StaticClass()->GetPathName());
 		}
+		else if (BaseLightNode->IsA<UInterchangeSkyLightNode>())
+		{
+			ActorFactoryNode->SetCustomActorClassName(ASkyLight::StaticClass()->GetPathName());
+		}
 		else
 		{
 			ActorFactoryNode->SetCustomActorClassName(APointLight::StaticClass()->GetPathName());
diff --git a/Engine/Plugins/Interchange/Runtime/Source/Pipelines/Private/InterchangeGenericTexturePipeline.cpp b/Engine/Plugins/Interchange/Runtime/Source/Pipelines/Private/InterchangeGenericTexturePipeline.cpp
--- a/Engine/Plugins/Interchange/Runtime/Source/Pipelines/Private/InterchangeGenericTexturePipeline.cpp
+++ b/Engine/Plugins/Interchange/Runtime/Source/Pipelines/Private/InterchangeGenericTexturePipeline.cpp
@@ -253,11 +253,18 @@ UInterchangeTextureFactoryNode* UInterchangeGenericTexturePipeline::HandleCreati
 {
 	UClass* FactoryClass = UE::Interchange::Private::GetDefaultFactoryClassFromTextureNodeClass(TextureNode->GetClass());
 
+	// Handle the Texture2D as a TextureCube if we should
 	TOptional<FString> SourceFile = TextureNode->GetPayLoadKey();
 #if WITH_EDITORONLY_DATA
 	if (FactoryClass == UInterchangeTexture2DFactoryNode::StaticClass())
-	{ 
-		if (SourceFile)
+	{
+		bool bForceLongLatCubemap = false;
+		const UInterchangeTexture2DNode* Texture2DNode = Cast<const UInterchangeTexture2DNode>(TextureNode);
+		if (Texture2DNode && Texture2DNode->GetForceLongLatCubemap(bForceLongLatCubemap) && bForceLongLatCubemap)
+		{
+			FactoryClass = UInterchangeTextureCubeFactoryNode::StaticClass();
+		}
+		else if (SourceFile)
 		{
 			const FString Extension = FPaths::GetExtension(SourceFile.GetValue()).ToLower();
 			if (FileExtensionsToImportAsLongLatCubemap.Contains(Extension))
diff --git a/Engine/Plugins/Runtime/USDCore/Source/USDUtilities/Private/USDPrimConversion.cpp b/Engine/Plugins/Runtime/USDCore/Source/USDUtilities/Private/USDPrimConversion.cpp
--- a/Engine/Plugins/Runtime/USDCore/Source/USDUtilities/Private/USDPrimConversion.cpp
+++ b/Engine/Plugins/Runtime/USDCore/Source/USDUtilities/Private/USDPrimConversion.cpp
@@ -103,6 +103,7 @@
 #include "pxr/usd/usdGeom/xformCommonAPI.h"
 #include "pxr/usd/usdLux/diskLight.h"
 #include "pxr/usd/usdLux/distantLight.h"
+#include "pxr/usd/usdLux/domeLight.h"
 #include "pxr/usd/usdLux/lightAPI.h"
 #include "pxr/usd/usdLux/rectLight.h"
 #include "pxr/usd/usdLux/shadowAPI.h"
@@ -1872,6 +1873,41 @@ UsdToUnreal::FPropertyTrackReader UsdToUnreal::CreatePropertyTrackReader(
 				}
 			}
 		}
+		else if (pxr::UsdLuxDomeLight DomeLight{UsdPrim})
+		{
+			if (PropertyPath == UnrealIdentifiers::IntensityPropertyName)
+			{
+				pxr::UsdAttribute IntensityAttr = SphereLight.GetIntensityAttr();
+				pxr::UsdAttribute ExposureAttr = SphereLight.GetExposureAttr();
+
+				if (IntensityAttr && ExposureAttr)
+				{
+					// Default values directly from pxr/usd/usdLux/schema.usda
+					float DefaultUsdIntensity = 1.0f;
+					float DefaultUsdExposure = 0.0f;
+
+					IntensityAttr.Get<float>(&DefaultUsdIntensity);
+					ExposureAttr.Get<float>(&DefaultUsdExposure);
+
+					float Default = UsdToUnreal::ConvertLightIntensityAttr(DefaultUsdIntensity, DefaultUsdExposure);
+
+					Reader.FloatReader = [IntensityAttr, ExposureAttr, Default, StageInfo](double UsdTimeCode)
+					{
+						float Result = Default;
+
+						float UsdIntensity = 1.0f;
+						float UsdExposure = 0.0f;
+						if (IntensityAttr.Get<float>(&UsdIntensity, UsdTimeCode) && ExposureAttr.Get<float>(&UsdExposure, UsdTimeCode))
+						{
+							Result = UsdToUnreal::ConvertLightIntensityAttr(UsdIntensity, UsdExposure);
+						}
+
+						return Result;
+					};
+					return Reader;
+				}
+			}
+		}
 	}
 	else if (const pxr::UsdMediaSpatialAudio Audio{Prim})
 	{
@@ -6434,7 +6470,7 @@ TArray<UE::FUsdAttribute> UsdUtils::GetAttributesForProperty(const UE::FUsdPrim&
 				UE::FUsdAttribute{UsdPrim.GetAttribute(pxr::UsdLuxTokens->inputsRadius)}
 			};
 		}
-		else if (UsdPrim.IsA<pxr::UsdLuxDistantLight>())
+		else if (UsdPrim.IsA<pxr::UsdLuxDistantLight>() || UsdPrim.IsA<pxr::UsdLuxDomeLight>())
 		{
 			return {
 				UE::FUsdAttribute{UsdPrim.GetAttribute(pxr::UsdLuxTokens->inputsIntensity)},
