[Interchange-USD] Adding Point Instancer support.
Adding Nested Mesh/xForm/PointInstancer support (within PointInstancers).
Adding LODContainer handling support (within PointInstancers)

#rnx

diff --git a/Engine/Plugins/Interchange/Extensions/OpenUSD/Source/Import/Private/InterchangeUsdTranslator.cpp b/Engine/Plugins/Interchange/Extensions/OpenUSD/Source/Import/Private/InterchangeUsdTranslator.cpp
--- a/Engine/Plugins/Interchange/Extensions/OpenUSD/Source/Import/Private/InterchangeUsdTranslator.cpp
+++ b/Engine/Plugins/Interchange/Extensions/OpenUSD/Source/Import/Private/InterchangeUsdTranslator.cpp
@@ -38,11 +38,13 @@
 
 #include "InterchangeCameraNode.h"
 #include "InterchangeLightNode.h"
+#include "InterchangeMeshLODContainerNode.h"
 #include "InterchangeManager.h"
 #include "InterchangeMaterialDefinitions.h"
 #include "InterchangeMaterialInstanceNode.h"
 #include "InterchangeMaterialReferenceNode.h"
 #include "InterchangeMeshNode.h"
+#include "InterchangeSceneComponentNodes.h"
 #include "InterchangeSceneNode.h"
 #include "InterchangeShaderGraphNode.h"
 #include "InterchangeTexture2DNode.h"
@@ -115,9 +117,13 @@ namespace UE::InterchangeUsdTranslator::Private
 	const static FString PrimitiveShapePrefix = TEXT("\\PrimitiveShape\\");
 	const static FString TexturePrefix = TEXT("\\Texture\\");
 	const static FString VolumePrefix = TEXT("\\Volume\\");
+	const static FString InstancerPrefix = TEXT("\\Instancer\\");
+	const static FString InstancePrefix = TEXT("\\Instance\\");
+	const static FString ComponentPrefix = TEXT("\\Component\\");
 
 	const static FString LODContainerSuffix = TEXT("LODContainer");
 	const static FString TwoSidedSuffix = TEXT("_TwoSided");
+	const static FString InstanceSuffix = TEXT("_Instance");
 
 	const static FString LODString = TEXT("LOD");
 
@@ -234,6 +240,38 @@ namespace UE::InterchangeUsdTranslator::Private
 		FTraversalInfo Info
 	);
 
+	struct FUidAndNameCounter
+	{
+		TMap<FString, int32> UidCounter;
+		TMap<FString, int32> NameCounter;
+		
+		void UpdateUid(FString& Uid)
+		{
+			int32& Counter = UidCounter.FindOrAdd(Uid);
+			if (Counter)
+			{
+				Uid = Uid + TEXT("-") + FString::FromInt(Counter);
+			}
+			Counter++;
+		}
+
+		void UpdateName(FString& Name)
+		{
+			int32& Counter = NameCounter.FindOrAdd(Name);
+			if (Counter)
+			{
+				Name = Name + TEXT("-") + FString::FromInt(Counter);
+			}
+			Counter++;
+		}
+
+		void Update(FString& Uid, FString& Name)
+		{
+			UpdateUid(Uid);
+			UpdateName(Name);
+		}
+	};
+
 	// clang-format off
 	const static TMap<FName, EInterchangePropertyTracks> PropertyNameToTrackType = {
 		// Common properties
@@ -400,6 +438,9 @@ namespace UE::InterchangeUsdTranslator::Private
 		// We stash the info we collected from each Volume prim path here, as we'll reuse it between translation and retrieving the payloads
 		TMap<FString, TArray<UsdUtils::FVolumePrimInfo>> PrimPathToVolumeInfo;
 
+		//To be able to create unique Instance specific Uids and Names. Specifically for PointInstancer support:
+		FUidAndNameCounter UidAndNameCounter;
+
 	private:
 		TMap<FString, TArray<FUsdMaterialXShaderGraph::FGeomProp>> MaterialUidToGeomProps;
 
@@ -1372,6 +1413,323 @@ namespace UE::InterchangeUsdTranslator::Private
 		const FString* BaseParameterName = nullptr;
 	};
 
+	//Helper to traverse PointInstancer subhierarchy and process it into UInterchangeSceneComponents and InstancedStaticMeshComponent.
+	struct FISMComponentTraversalHelper
+	{
+		UInterchangeUSDTranslatorImpl& TranslatorImpl;
+		UInterchangeBaseNodeContainer& NodeContainer;
+		UInterchangeUsdTranslatorSettings* TranslatorSettings;
+
+		FISMComponentTraversalHelper(
+			UInterchangeUSDTranslatorImpl& InTranslatorImpl,
+			UInterchangeBaseNodeContainer& InNodeContainer,
+			UInterchangeUsdTranslatorSettings* InTranslatorSettings)
+			: TranslatorImpl(InTranslatorImpl)
+			, NodeContainer(InNodeContainer)
+			, TranslatorSettings(InTranslatorSettings)
+		{
+
+		}
+		
+		UInterchangeSceneComponentNode* CreateSceneComponent(UInterchangeBaseNode* ParentBaseNode, FString& UniqueId, FString& DisplayLabel, const FTransform& LocalTransform, bool bVisible, UClass* ComponentClass = UInterchangeSceneComponentNode::StaticClass())
+		{
+			TranslatorImpl.UidAndNameCounter.Update(UniqueId, DisplayLabel);
+
+			UInterchangeSceneComponentNode* SceneComponentNode = NewObject<UInterchangeSceneComponentNode>(&NodeContainer, ComponentClass);
+			NodeContainer.SetupNode(SceneComponentNode, UniqueId, DisplayLabel, EInterchangeNodeContainerType::TranslatedScene, ParentBaseNode ? ParentBaseNode->GetUniqueID() : TEXT(""));
+
+			SceneComponentNode->SetCustomLocalTransform(LocalTransform);
+			SceneComponentNode->SetCustomComponentVisibility(bVisible);
+
+			if (UInterchangeSceneComponentNode* ParentSceneComponentNode = Cast<UInterchangeSceneComponentNode>(ParentBaseNode))
+			{
+				ParentSceneComponentNode->AddComponentUid(SceneComponentNode->GetUniqueID());
+			}
+			else if (UInterchangeSceneNode* ParentSceneNode = Cast<UInterchangeSceneNode>(ParentBaseNode))
+			{
+				ParentSceneNode->AddComponentUid(SceneComponentNode->GetUniqueID());
+			}
+
+			return SceneComponentNode;
+		}
+		
+		UInterchangeSceneComponentNode* CreateInstanceSceneComponent(UInterchangeBaseNode* ParentBaseNode, FString& UniqueId, FString& DisplayLabel, int32 InstanceCounter, const FTransform& LocalTransform, bool bVisible)
+		{
+			FString InstanceUid = UniqueId + InstancePrefix + FString::FromInt(InstanceCounter);
+			FString InstanceDisplayLabel = DisplayLabel + InstanceSuffix + FString::FromInt(InstanceCounter);
+
+			//Note: For handling the point instancer's instance's transform we create a component here. This is an extra introduced component compared to the original hierarchy.
+			return CreateSceneComponent(ParentBaseNode, InstanceUid, InstanceDisplayLabel, LocalTransform, bVisible);
+		}
+
+		FString ConvertLodContainerPrim(const UE::FUsdPrim& Prim, FString& ActiveVariantUSDPrimPath, const FTraversalInfo& Info)
+		{
+			FString LODContainerUid = LODContainerSuffix + TEXT("\\") + Prim.GetPrimPath().GetString();
+			FString ComponentDisplayLabel(Prim.GetName().ToString());
+
+			UE::FUsdVariantSets VariantSets = Prim.GetVariantSets();
+			const UE::FUsdVariantSet LODVariantSet = VariantSets.GetVariantSet(LODString);
+			if (!LODVariantSet.IsValid())
+			{
+				return TEXT("");
+			}
+
+			const FString ActiveVariant = LODVariantSet.GetVariantSelection();
+			if (!IsValidLODName(ActiveVariant))
+			{
+				return TEXT("");
+			}
+
+			if (const UInterchangeBaseNode* Node = NodeContainer.GetNode(LODContainerUid))
+			{
+				if (UE::FUsdPrim ActiveLODMeshPrim = GetLODMesh(Prim, ActiveVariant))
+				{
+					ActiveVariantUSDPrimPath = ActiveLODMeshPrim.GetPrimPath().GetString();
+				}
+
+				return LODContainerUid;
+			}
+
+			UInterchangeMeshLODContainerNode* LODContainerComponentNode = NewObject<UInterchangeMeshLODContainerNode>(&NodeContainer);
+			NodeContainer.SetupNode(LODContainerComponentNode, LODContainerUid, ComponentDisplayLabel, EInterchangeNodeContainerType::TranslatedScene);
+
+			bool bSwitchedFromInitialVariant = false;
+			TArray<FString> VariantNames = LODVariantSet.GetVariantNames();
+			for (const FString& VariantName : VariantNames)
+			{
+				// The active variant will be parsed via regular traversal
+				if (!IsValidLODName(VariantName))
+				{
+					continue;
+				}
+
+				{
+					// For creating the scene nodes themselves we'll switch the active variant on the currently opened stage
+					// (still using the session layer to minimize impact to the actual layer). This mainly so that we can retrieve
+					// and fetch and cache the correct material bindings for the LOD meshes. Later on we'll use separate stages
+					// with population masks to read the LODs concurrently, and we won't be able to resolve material bindings
+					UE::FUsdEditContext Context{ TranslatorImpl.UsdStage, TranslatorImpl.UsdStage.GetSessionLayer() };
+
+					bool bSwitchedVariant = VariantSets.SetSelection(LODString, VariantName);
+					if (!bSwitchedVariant)
+					{
+						continue;
+					}
+					bSwitchedFromInitialVariant = true;
+				}
+
+				if (UE::FUsdPrim LODMeshPrim = GetLODMesh(Prim, VariantName))
+				{
+					LODContainerComponentNode->AddMeshLODNodeUid(TranslatorImpl.AddMeshNode(LODMeshPrim, TranslatorSettings, NodeContainer, Info));
+
+					if (VariantName == ActiveVariant)
+					{
+						ActiveVariantUSDPrimPath = LODMeshPrim.GetPrimPath().GetString();
+					}
+				}
+				else
+				{
+					UE_LOG(LogTemp, Warning, TEXT("Failed to parse a LOD Mesh from variant '%s' of prim '%s'. For automatic parsing of LODs, make sure there is a single Mesh prim within the variant, named exactly as the variant itself (e.g. 'LOD0', 'LOD1', etc.)"), *VariantName, *Prim.GetPrimPath().GetString());
+				}
+			}
+
+			// Put the active variant back to what it originally was
+			if (bSwitchedFromInitialVariant)
+			{
+				{
+					UE::FUsdEditContext Context{ TranslatorImpl.UsdStage, TranslatorImpl.UsdStage.GetSessionLayer() };
+
+					const bool bRestoredSelection = VariantSets.SetSelection(LODString, ActiveVariant);
+					ensure(bRestoredSelection);
+				}
+			}
+
+			return LODContainerUid;
+		}
+
+		void TraverseNestedPointInstancerHierarchy(
+			const UE::FUsdPrim& Prim,
+			FTraversalInfo Info,
+			UInterchangeSceneComponentNode* ParentSceneComponentNode
+		)
+		{
+			const FUsdStageInfo StageInfo(Prim.GetStage());
+
+			FTransform xFormTransform = FTransform::Identity;
+			bool bResetTransformStack = false;
+			UsdToUnreal::ConvertXformable(Prim.GetStage(), UE::FUsdTyped(Prim), xFormTransform, UsdUtils::GetEarliestTimeCode(), &bResetTransformStack);
+
+			FString ComponentUid = ComponentPrefix + Prim.GetPrimPath().GetString();
+			FString ComponentDisplayLabel(Prim.GetName().ToString());
+
+			Info.UpdateWithCurrentPrim(Prim);
+
+			if (Prim.IsA(TEXT("PointInstancer")))
+			{
+				UInterchangeSceneComponentNode* SceneComponentNode = CreateSceneComponent(ParentSceneComponentNode, ComponentUid, ComponentDisplayLabel, xFormTransform, Info.bVisible);
+
+				if (SceneComponentNode)
+				{
+					TraversePointInstancerHierarchy(Prim, Info, SceneComponentNode);
+				}
+			}
+			else if (Info.bIsLODContainer)
+			{
+				FString ActiveVariantUSDPrimPath;
+				FString LODContainerUid = ConvertLodContainerPrim(Prim, ActiveVariantUSDPrimPath, Info);
+
+				//This could be a basic StaticMesh, but currently we only have support for InstancedStaticMeshes. (within the context of Components)
+				UInterchangeInstancedStaticMeshComponentNode* ISMComponentNode = Cast<UInterchangeInstancedStaticMeshComponentNode>(CreateSceneComponent(ParentSceneComponentNode, ComponentUid, ComponentDisplayLabel, FTransform::Identity, Info.bVisible, UInterchangeInstancedStaticMeshComponentNode::StaticClass()));
+
+				if (!ISMComponentNode)
+				{
+					return;
+				}
+
+				ISMComponentNode->AddInstanceTransform(xFormTransform);
+				if (LODContainerUid.Len())
+				{
+					ISMComponentNode->SetCustomInstancedAssetUid(LODContainerUid);
+				}
+
+				TArray<FUsdPrim> Children = Prim.GetChildren();
+				if (Children.Num())
+				{
+					//Note: For handling the point instancer's instance's transform we create a component here. This is an extra introduced component compared to the original hierarchy.
+					UInterchangeSceneComponentNode* SceneComponentNode = CreateInstanceSceneComponent(ISMComponentNode, ComponentUid, ComponentDisplayLabel, 0 /*InstanceCounter*/, xFormTransform, Info.bVisible);
+
+					if (!SceneComponentNode)
+					{
+						return;
+					}
+
+					for (FUsdPrim& Child : Children)
+					{
+						if (ActiveVariantUSDPrimPath == Child.GetPrimPath().GetString())
+						{
+							TArray<FUsdPrim> ActiveVariantChildren = Child.GetChildren();
+							for (const FUsdPrim& ActiveVariantChild : ActiveVariantChildren)
+							{
+								TraverseNestedPointInstancerHierarchy(ActiveVariantChild, Info, SceneComponentNode);
+							}
+						}
+						else
+						{
+							TraverseNestedPointInstancerHierarchy(Child, Info, SceneComponentNode);
+						}
+					}
+				}
+			}
+			else if (Prim.IsA(TEXT("Xform")))
+			{
+				UInterchangeSceneComponentNode* SceneComponentNode = CreateSceneComponent(ParentSceneComponentNode, ComponentUid, ComponentDisplayLabel, xFormTransform, Info.bVisible, UInterchangeSceneComponentNode::StaticClass());
+
+				if (!SceneComponentNode)
+				{
+					return;
+				}
+
+				TArray<FUsdPrim> Children = Prim.GetChildren();
+				for (const FUsdPrim& Child : Children)
+				{
+					TraverseNestedPointInstancerHierarchy(Child, Info, SceneComponentNode);
+				}
+			}
+			else if (Prim.IsA(TEXT("Mesh")))
+			{
+				//This could be a basic StaticMesh, but currently we only have support for InstancedStaticMeshes. (within the context of Components)
+				UInterchangeInstancedStaticMeshComponentNode* ISMComponentNode = Cast<UInterchangeInstancedStaticMeshComponentNode>(CreateSceneComponent(ParentSceneComponentNode, ComponentUid, ComponentDisplayLabel, FTransform::Identity, Info.bVisible, UInterchangeInstancedStaticMeshComponentNode::StaticClass()));
+
+				if (!ISMComponentNode)
+				{
+					return;
+				}
+
+				ISMComponentNode->AddInstanceTransform(xFormTransform);
+				ISMComponentNode->SetCustomInstancedAssetUid(TranslatorImpl.AddMeshNode(Prim, TranslatorSettings, NodeContainer, Info));
+			}
+			else
+			{
+				UE_LOG(
+					LogUsd,
+					Warning,
+					TEXT("Unexpected Prim type [%s] in Instancer sub components."),
+					*Prim.GetTypeName().ToString()
+				);
+			}
+		}
+
+		void TraversePointInstancerHierarchy(
+			const UE::FUsdPrim& Prim,
+			FTraversalInfo Info,
+			UInterchangeBaseNode* ParentBaseNode
+		)
+		{
+			UsdUtils::FUsdPointInstancerInstancesData InstancesData;
+			UsdUtils::GetPointInstancerInstancesData(Prim, InstancesData);
+
+			const FUsdStageInfo StageInfo(Prim.GetStage());
+
+			for (int32 PrototypeIndex = 0; PrototypeIndex < InstancesData.Prototypes.Num(); ++PrototypeIndex)
+			{
+				UsdUtils::FUsdPointInstancerInstancesData::FUsdPrototypeInstances& PrototypeInstanceData = InstancesData.Prototypes[PrototypeIndex];
+				TArray<FTransform>& PrototypeInstanceTransforms = PrototypeInstanceData.InstanceTransforms;
+				FString PrototypePathString = PrototypeInstanceData.PathString;
+
+				UE::FUsdPrim PrototypePrim = Prim.GetStage().GetPrimAtPath(UE::FSdfPath{ *PrototypePathString });
+				FString InstancerUid = InstancerPrefix + PrototypePathString;
+				FString InstancerDisplayLabel(PrototypePrim.GetName().ToString());
+
+				FTraversalInfo PrototypePrimInfo = Info;
+				PrototypePrimInfo.UpdateWithCurrentPrim(PrototypePrim);
+
+				if (PrototypePrim.IsA(TEXT("PointInstancer")) || PrototypePrim.IsA(TEXT("Xform")))
+				{
+					for (int32 InstanceCounter = 0; InstanceCounter < PrototypeInstanceTransforms.Num(); ++InstanceCounter)
+					{
+						//Note: For handling the point instancer's instance's transform we create a component here. This is an extra introduced component compared to the original hierarchy.
+						UInterchangeSceneComponentNode* SceneComponentNode = CreateInstanceSceneComponent(ParentBaseNode, InstancerUid, InstancerDisplayLabel, InstanceCounter, PrototypeInstanceTransforms[InstanceCounter], PrototypePrimInfo.bVisible);
+
+						if (SceneComponentNode)
+						{
+							TraverseNestedPointInstancerHierarchy(PrototypePrim, PrototypePrimInfo, SceneComponentNode);
+						}
+					}
+				}
+				else if (PrototypePrim.IsA(TEXT("Mesh")))
+				{
+					FTransform xFormTransform = FTransform::Identity;
+					bool bResetTransformStack = false;
+					UsdToUnreal::ConvertXformable(PrototypePrim.GetStage(), UE::FUsdTyped(PrototypePrim), xFormTransform, UsdUtils::GetEarliestTimeCode(), &bResetTransformStack);
+
+					for (FTransform& Transform : PrototypeInstanceTransforms)
+					{
+						Transform = xFormTransform * Transform;
+					}
+
+					UInterchangeInstancedStaticMeshComponentNode* ISMComponentNode = Cast<UInterchangeInstancedStaticMeshComponentNode>(CreateSceneComponent(ParentBaseNode, InstancerUid, InstancerDisplayLabel, FTransform::Identity, PrototypePrimInfo.bVisible, UInterchangeInstancedStaticMeshComponentNode::StaticClass()));
+
+					if (!ISMComponentNode)
+					{
+						continue;
+					}
+					
+					ISMComponentNode->AddInstanceTransforms(PrototypeInstanceTransforms);
+					ISMComponentNode->SetCustomInstancedAssetUid(TranslatorImpl.AddMeshNode(PrototypePrim, TranslatorSettings, NodeContainer, PrototypePrimInfo));
+				}
+				else
+				{
+					UE_LOG(
+						LogUsd,
+						Warning,
+						TEXT("Unexpected Prim type [%s] in PointInstancer."),
+						*PrototypePrim.GetTypeName().ToString()
+					);
+				}
+			}
+		}
+	};
+
 	FString UInterchangeUSDTranslatorImpl::AddMaterialNode(
 		const UE::FUsdPrim& Prim,
 		UInterchangeUsdTranslatorSettings* TranslatorSettings,
@@ -3629,7 +3987,7 @@ namespace UE::InterchangeUsdTranslator::Private
 		}
 		else
 		{
-			const static TSet<FName> KnownUnsupported = {TEXT("SpatialAudio"), TEXT("PointInstancer"), TEXT("BasisCurves")};
+			const static TSet<FName> KnownUnsupported = {TEXT("SpatialAudio"), TEXT("BasisCurves")};
 
 			if (KnownUnsupported.Contains(TypeName))
 			{
@@ -3781,9 +4139,17 @@ namespace UE::InterchangeUsdTranslator::Private
 		{
 			Info.ParentNode = SceneNode;
 
-			for (const FUsdPrim& ChildPrim : CheckLodAPIAndGetChildren(Prim, SceneNode))
+			if (Prim.IsA(TEXT("PointInstancer")))
+			{
+				FISMComponentTraversalHelper ISMComponentTraversalHelper(TranslatorImpl, NodeContainer, TranslatorSettings);
+				ISMComponentTraversalHelper.TraversePointInstancerHierarchy(Prim, Info, SceneNode);
+			}
+			else
 			{
-				Traverse(ChildPrim, TranslatorImpl, NodeContainer, TranslatorSettings, Info);
+				for (const FUsdPrim& ChildPrim : CheckLodAPIAndGetChildren(Prim, SceneNode))
+				{
+					Traverse(ChildPrim, TranslatorImpl, NodeContainer, TranslatorSettings, Info);
+				}
 			}
 		}
 
diff --git a/Engine/Plugins/Runtime/USDCore/Source/USDUtilities/Private/USDGeomMeshConversion.cpp b/Engine/Plugins/Runtime/USDCore/Source/USDUtilities/Private/USDGeomMeshConversion.cpp
--- a/Engine/Plugins/Runtime/USDCore/Source/USDUtilities/Private/USDGeomMeshConversion.cpp
+++ b/Engine/Plugins/Runtime/USDCore/Source/USDUtilities/Private/USDGeomMeshConversion.cpp
@@ -6892,6 +6892,34 @@ bool UsdUtils::GetPointInstancerTransforms(
 	return true;
 }
 
+void UsdUtils::GetPointInstancerInstancesData(const pxr::UsdPrim& Prim, FUsdPointInstancerInstancesData& InstancesData)
+{
+	pxr::UsdGeomPointInstancer PointInstancer(Prim);
+	if (!PointInstancer)
+	{
+		return;
+	}
+
+	pxr::SdfPathVector PrototypePaths;
+	if (!PointInstancer.GetPrototypesRel().GetTargets(&PrototypePaths))
+	{
+		return;
+	}
+
+	UE::FUsdStage Stage(Prim.GetStage());
+	const FUsdStageInfo StageInfo(Prim.GetStage());
+
+	InstancesData.Prototypes.SetNum(PrototypePaths.size());
+
+	for (int32 PrototypeIndex = 0; PrototypeIndex < PrototypePaths.size(); ++PrototypeIndex)
+	{
+		const pxr::SdfPath& PrototypePath = PrototypePaths[PrototypeIndex];
+		InstancesData.Prototypes[PrototypeIndex].PathString = UsdToUnreal::ConvertPath(PrototypePath);
+
+		UsdUtils::GetPointInstancerTransforms(StageInfo, PointInstancer, PrototypeIndex, UsdUtils::GetEarliestTimeCode(), InstancesData.Prototypes[PrototypeIndex].InstanceTransforms);
+	}
+}
+
 bool UsdUtils::IsAnimatedMesh(const pxr::UsdPrim& Prim)
 {
 	FScopedUsdAllocs UsdAllocs;
diff --git a/Engine/Plugins/Runtime/USDCore/Source/USDUtilities/Public/USDGeomMeshConversion.h b/Engine/Plugins/Runtime/USDCore/Source/USDUtilities/Public/USDGeomMeshConversion.h
--- a/Engine/Plugins/Runtime/USDCore/Source/USDUtilities/Public/USDGeomMeshConversion.h
+++ b/Engine/Plugins/Runtime/USDCore/Source/USDUtilities/Public/USDGeomMeshConversion.h
@@ -521,6 +521,20 @@ namespace UsdUtils
 		TArray<FTransform>& OutInstanceTransforms
 	);
 
+	struct FUsdPointInstancerInstancesData
+	{
+		struct FUsdPrototypeInstances
+		{
+			FString PathString;
+			TArray<FTransform> InstanceTransforms;
+		};
+
+		TArray<FUsdPrototypeInstances> Prototypes;
+	};
+
+	/** Acquires the Prototypes from the PointInstancer and the related InstanceTransforms. */
+	USDUTILITIES_API void GetPointInstancerInstancesData(const pxr::UsdPrim& Prim, FUsdPointInstancerInstancesData& InstancesData);
+
 	/** Returns true if Prim is a GeomMesh with animated attributes */
 	USDUTILITIES_API bool IsAnimatedMesh(const pxr::UsdPrim& Prim);
 
