USD: Convert prim metadata into into Interchange translated node user attributes;

Also add a regex filter to select only desired metadata.

diff --git a/Engine/Plugins/Interchange/Extensions/OpenUSD/Source/Import/Private/InterchangeUsdTranslator.cpp b/Engine/Plugins/Interchange/Extensions/OpenUSD/Source/Import/Private/InterchangeUsdTranslator.cpp
--- a/Engine/Plugins/Interchange/Extensions/OpenUSD/Source/Import/Private/InterchangeUsdTranslator.cpp
+++ b/Engine/Plugins/Interchange/Extensions/OpenUSD/Source/Import/Private/InterchangeUsdTranslator.cpp
@@ -2269,7 +2269,7 @@ namespace UE::InterchangeUsdTranslator::Private
 			NodeContainer.SetupNode(LightNode, NodeUid, NodeName, EInterchangeNodeContainerType::TranslatedAsset);
 			LightNode->SetAssetName(NodeName);
 
-			// TODO: Maybe only authore these in case there were in fact authored opinions on the prims? Otherwise we're just authoring default values
+			// TODO: Maybe only author these in case there were in fact authored opinions on the prims? Otherwise we're just authoring default values
 			LightNode->SetCustomLightColor(Color);
 			LightNode->SetCustomTemperature(Temperature);
 			LightNode->SetCustomUseTemperature(UseTemperature);
@@ -3822,15 +3822,253 @@ namespace UE::InterchangeUsdTranslator::Private
 #endif	  // WITH_EDITOR
 	}
 
-	template<typename T>
-	void TranslateAttribute(const UE::FUsdAttribute& Attr, const FString& AttrName, UInterchangeSceneNode* Node)
+	template<typename USDType, typename InterchangeType = USDType>
+	InterchangeType CoerceToInterchangeType(const USDType& USDValue)
 	{
-		T Value;
-		if (Attr.Get<T>(Value))	   // Always check for an opinion on the default time code
+		if constexpr (std::is_same_v<USDType, FMatrix2D> && std::is_same_v<InterchangeType, FMatrix44d>)
 		{
-			TOptional<FString> PayloadKey;
-			UInterchangeUserDefinedAttributesAPI::CreateUserDefinedAttribute(Node, AttrName, Value, PayloadKey);
+			FMatrix44d Result{EForceInit::ForceInitToZero};
+			Result.M[0][0] = USDValue.Row0.X;
+			Result.M[0][1] = USDValue.Row0.Y;
+
+			Result.M[1][0] = USDValue.Row1.X;
+			Result.M[1][1] = USDValue.Row1.Y;
+			return Result;
+		}
+		else if constexpr (std::is_same_v<USDType, FMatrix3D> && std::is_same_v<InterchangeType, FMatrix44d>)
+		{
+			FMatrix44d Result{EForceInit::ForceInitToZero};
+			Result.M[0][0] = USDValue.Row0.X;
+			Result.M[0][1] = USDValue.Row0.Y;
+			Result.M[0][2] = USDValue.Row0.Z;
+
+			Result.M[1][0] = USDValue.Row1.X;
+			Result.M[1][1] = USDValue.Row1.Y;
+			Result.M[1][2] = USDValue.Row1.Z;
+
+			Result.M[2][0] = USDValue.Row2.X;
+			Result.M[2][1] = USDValue.Row2.Y;
+			Result.M[2][2] = USDValue.Row2.Z;
+			return Result;
 		}
+		else
+		{
+			return (InterchangeType)USDValue;
+		}
+	};
+
+	template<typename USDType, typename InterchangeType = USDType>
+	void AddUserAttribute(const FString& KeyName, const UE::FVtValue& VtValue, UInterchangeSceneNode* Node)
+	{
+		const USDType USDValue = VtValue.Get<USDType>();
+		const TOptional<FString> PayloadKey;
+
+		if constexpr (std::is_same_v<InterchangeType, USDType>)
+		{
+			UInterchangeUserDefinedAttributesAPI::CreateUserDefinedAttribute(Node, KeyName, USDValue, PayloadKey);
+		}
+		else if constexpr (TIsTArray<USDType>::Value)
+		{
+			static_assert(TIsTArray<InterchangeType>::Value);
+
+			InterchangeType InterchangeArray;
+			InterchangeArray.SetNum(USDValue.Num());
+			for (int32 Index = 0; Index < USDValue.Num(); ++Index)
+			{
+				InterchangeArray[Index] = CoerceToInterchangeType<typename USDType::ElementType, typename InterchangeType::ElementType>(USDValue[Index]);
+			}
+
+			UInterchangeUserDefinedAttributesAPI::CreateUserDefinedAttribute(Node, KeyName, InterchangeArray, PayloadKey);
+		}
+		else
+		{
+			const InterchangeType InterchangeValue = CoerceToInterchangeType<USDType, InterchangeType>(USDValue);
+			UInterchangeUserDefinedAttributesAPI::CreateUserDefinedAttribute(Node, KeyName, InterchangeValue, PayloadKey);
+		}
+	};
+
+	// Note: We could have just used the Cpp type names here instead of the value type token, but doing the
+	// latter lets us handle the color types as FLinearColor, which is probably more useful
+	// References:
+	// - https://openusd.org/docs/api/_usd__page__datatypes.html
+	// - EAttributeTypes declaration
+	// - Engine/Source/Runtime/Interchange/Core/Private/Tests/StorageTest.cpp
+	using AddUserAttributeFunc = TFunction<void(const FString&, const UE::FVtValue&, UInterchangeSceneNode*)>;
+	static TMap<FName, AddUserAttributeFunc> AddUserAttributeFuncs = {
+		{TEXT("bool"),  					AddUserAttribute<bool>},
+
+		{TEXT("uchar"), 					AddUserAttribute<uint8>},
+		{TEXT("unsigned char"), 			AddUserAttribute<uint8>},
+		{TEXT("int"), 						AddUserAttribute<int32>},
+		{TEXT("unsigned int"), 				AddUserAttribute<uint32>},
+		{TEXT("uint"), 						AddUserAttribute<uint32>},
+		{TEXT("int64"), 					AddUserAttribute<int64>},
+		{TEXT("__int64"), 					AddUserAttribute<int64>},
+		{TEXT("unsigned __int64"), 			AddUserAttribute<uint64>},
+		{TEXT("uint64"), 					AddUserAttribute<uint64>},
+
+		{TEXT("half"),	 					AddUserAttribute<FFloat16>},
+		{TEXT("pxr_half::half"),	 		AddUserAttribute<FFloat16>},
+		{TEXT("float"), 					AddUserAttribute<float>},
+		{TEXT("double"), 					AddUserAttribute<double>},
+		{TEXT("timecode"), 					AddUserAttribute<FSdfTimeCode, double>},
+		{TEXT("SdfTimeCode"), 				AddUserAttribute<FSdfTimeCode, double>},
+
+		{TEXT("string"), 					AddUserAttribute<FString>},
+		{TEXT("token"), 					AddUserAttribute<FName>},
+		{TEXT("TfToken"), 					AddUserAttribute<FName>},
+		{TEXT("asset"), 					AddUserAttribute<FSdfAssetPath, FString>},
+		{TEXT("SdfAssetPath"), 				AddUserAttribute<FSdfAssetPath, FString>},
+
+		{TEXT("matrix2d"), 					AddUserAttribute<FMatrix2D, FMatrix44d>},
+		{TEXT("GfMatrix2d"), 				AddUserAttribute<FMatrix2D, FMatrix44d>},
+		{TEXT("matrix3d"), 					AddUserAttribute<FMatrix3D, FMatrix44d>},
+		{TEXT("GfMatrix3d"), 				AddUserAttribute<FMatrix3D, FMatrix44d>},
+		{TEXT("matrix4d"), 					AddUserAttribute<FMatrix44d>},
+		{TEXT("GfMatrix4d"), 				AddUserAttribute<FMatrix44d>},
+		{TEXT("frame4d"), 					AddUserAttribute<FMatrix44d>},
+
+		{TEXT("quath"), 					AddUserAttribute<FQuat4h, FQuat4f>},
+		{TEXT("GfQuath"), 					AddUserAttribute<FQuat4h, FQuat4f>},
+		{TEXT("quatf"), 					AddUserAttribute<FQuat4f>},
+		{TEXT("GfQuatf"), 					AddUserAttribute<FQuat4f>},
+		{TEXT("quatd"), 					AddUserAttribute<FQuat4d>},
+		{TEXT("GfQuatd"), 					AddUserAttribute<FQuat4d>},
+
+		{TEXT("half2"), 					AddUserAttribute<FVector2DHalf>},
+		{TEXT("GfVec2h"), 					AddUserAttribute<FVector2DHalf>},
+		{TEXT("float2"), 					AddUserAttribute<FVector2f>},
+		{TEXT("GfVec2f"), 					AddUserAttribute<FVector2f>},
+		{TEXT("double2"), 					AddUserAttribute<FVector2d>},
+		{TEXT("GfVec2d"), 					AddUserAttribute<FVector2d>},
+		{TEXT("int2"), 						AddUserAttribute<FIntPoint>},
+		{TEXT("GfVec2i"), 					AddUserAttribute<FIntPoint>},
+
+		{TEXT("half3"), 					AddUserAttribute<FVector3h, FVector3f>},
+		{TEXT("GfVec3h"), 					AddUserAttribute<FVector3h, FVector3f>},
+		{TEXT("point3h"), 					AddUserAttribute<FVector3h, FVector3f>},
+		{TEXT("normal3h"), 					AddUserAttribute<FVector3h, FVector3f>},
+		{TEXT("vector3h"), 					AddUserAttribute<FVector3h, FVector3f>},
+		{TEXT("color3h"), 					AddUserAttribute<FVector3h, FLinearColor>},
+
+		{TEXT("float3"), 					AddUserAttribute<FVector3f>},
+		{TEXT("GfVec3f"), 					AddUserAttribute<FVector3f>},
+		{TEXT("point3f"), 					AddUserAttribute<FVector3f>},
+		{TEXT("normal3f"), 					AddUserAttribute<FVector3f>},
+		{TEXT("vector3f"), 					AddUserAttribute<FVector3f>},
+		{TEXT("color3f"), 					AddUserAttribute<FVector3f, FLinearColor>},
+
+		{TEXT("double3"), 					AddUserAttribute<FVector3d>},
+		{TEXT("GfVec3d"), 					AddUserAttribute<FVector3d>},
+		{TEXT("point3d"), 					AddUserAttribute<FVector3d>},
+		{TEXT("normal3d"), 					AddUserAttribute<FVector3d>},
+		{TEXT("vector3d"), 					AddUserAttribute<FVector3d>},
+		{TEXT("color3d"), 					AddUserAttribute<FVector3d, FLinearColor>},
+
+		{TEXT("int3"), 						AddUserAttribute<FIntVector>},
+		{TEXT("GfVec3i"), 					AddUserAttribute<FIntVector>},
+
+		{TEXT("half4"), 					AddUserAttribute<FVector4h, FVector4f>},
+		{TEXT("GfVec4h"), 					AddUserAttribute<FVector4h, FVector4f>},
+		{TEXT("float4"), 					AddUserAttribute<FVector4f>},
+		{TEXT("GfVec4f"), 					AddUserAttribute<FVector4f>},
+		{TEXT("double4"), 					AddUserAttribute<FVector4d>},
+		{TEXT("GfVec4d"), 					AddUserAttribute<FVector4d>},
+		{TEXT("int4"), 						AddUserAttribute<FIntRect>},
+		{TEXT("GfVec4i"), 					AddUserAttribute<FIntRect>},
+		{TEXT("color4h"), 					AddUserAttribute<FVector4h, FLinearColor>},
+		{TEXT("color4f"), 					AddUserAttribute<FVector4f, FLinearColor>},
+		{TEXT("color4d"), 					AddUserAttribute<FVector4d, FLinearColor>},
+
+		{TEXT("bool[]"),  					AddUserAttribute<TArray<bool>>},
+		{TEXT("VtArray<bool>"),  			AddUserAttribute<TArray<bool>>},
+
+		{TEXT("uchar[]"), 					AddUserAttribute<TArray<uint8>>},
+		{TEXT("VtArray<unsigned char>"), 	AddUserAttribute<TArray<uint8>>},
+		{TEXT("int[]"), 					AddUserAttribute<TArray<int32>>},
+		{TEXT("VtArray<int>"), 				AddUserAttribute<TArray<int32>>},
+		{TEXT("uint[]"), 					AddUserAttribute<TArray<uint32>>},
+		{TEXT("VtArray<unsigned int>"), 	AddUserAttribute<TArray<uint32>>},
+		{TEXT("int64[]"), 					AddUserAttribute<TArray<int64>>},
+		{TEXT("VtArray<__int64>"), 			AddUserAttribute<TArray<int64>>},
+		{TEXT("uint64[]"), 					AddUserAttribute<TArray<uint64>>},
+		{TEXT("VtArray<unsigned __int64>"), AddUserAttribute<TArray<uint64>>},
+
+		{TEXT("half[]"),	 				AddUserAttribute<TArray<FFloat16>>},
+		{TEXT("VtArray<pxr_half::half>"),	AddUserAttribute<TArray<FFloat16>>},
+		{TEXT("float[]"), 					AddUserAttribute<TArray<float>>},
+		{TEXT("VtArray<float>"), 			AddUserAttribute<TArray<float>>},
+		{TEXT("double[]"), 					AddUserAttribute<TArray<double>>},
+		{TEXT("VtArray<double>"), 			AddUserAttribute<TArray<double>>},
+		{TEXT("timecode[]"), 				AddUserAttribute<TArray<FSdfTimeCode>, TArray<double>>},
+		{TEXT("VtArray<SdfTimeCode>"), 		AddUserAttribute<TArray<FSdfTimeCode>, TArray<double>>},
+
+		{TEXT("string[]"), 					AddUserAttribute<TArray<FString>>},
+		{TEXT("VtArray<string>"), 			AddUserAttribute<TArray<FString>>},
+		{TEXT("token[]"), 					AddUserAttribute<TArray<FName>>},
+		{TEXT("VtArray<TfToken>"), 			AddUserAttribute<TArray<FName>>},
+		{TEXT("asset[]"), 					AddUserAttribute<TArray<FSdfAssetPath>, TArray<FString>>},
+		{TEXT("VtArray<SdfAssetPath>"), 	AddUserAttribute<TArray<FSdfAssetPath>, TArray<FString>>},
+
+		{TEXT("matrix2d[]"), 				AddUserAttribute<TArray<FMatrix2D>, TArray<FMatrix44d>>},
+		{TEXT("VtArray<GfMatrix2d>"), 		AddUserAttribute<TArray<FMatrix2D>, TArray<FMatrix44d>>},
+		{TEXT("matrix3d[]"), 				AddUserAttribute<TArray<FMatrix3D>, TArray<FMatrix44d>>},
+		{TEXT("VtArray<GfMatrix3d>"), 		AddUserAttribute<TArray<FMatrix3D>, TArray<FMatrix44d>>},
+		{TEXT("matrix4d[]"), 				AddUserAttribute<TArray<FMatrix44d>>},
+		{TEXT("VtArray<GfMatrix4d>"), 		AddUserAttribute<TArray<FMatrix44d>>},
+		{TEXT("frame4d[]"), 				AddUserAttribute<TArray<FMatrix44d>>},
+
+		{TEXT("quath[]"), 					AddUserAttribute<TArray<FQuat4h>, TArray<FQuat4f>>},
+		{TEXT("VtArray<GfQuath>"), 			AddUserAttribute<TArray<FQuat4h>, TArray<FQuat4f>>},
+		{TEXT("quatf[]"), 					AddUserAttribute<TArray<FQuat4f>>},
+		{TEXT("VtArray<GfQuatf>"), 			AddUserAttribute<TArray<FQuat4f>>},
+		{TEXT("quatd[]"), 					AddUserAttribute<TArray<FQuat4d>>},
+		{TEXT("VtArray<GfQuatd>"), 			AddUserAttribute<TArray<FQuat4d>>},
+
+		{TEXT("half2[]"), 					AddUserAttribute<TArray<FVector2DHalf>>},
+		{TEXT("VtArray<GfVec2h>"), 			AddUserAttribute<TArray<FVector2DHalf>>},
+		{TEXT("float2[]"), 					AddUserAttribute<TArray<FVector2f>>},
+		{TEXT("VtArray<GfVec2f>"), 			AddUserAttribute<TArray<FVector2f>>},
+		{TEXT("double2[]"), 				AddUserAttribute<TArray<FVector2d>>},
+		{TEXT("VtArray<GfVec2d>"), 			AddUserAttribute<TArray<FVector2d>>},
+		{TEXT("int2[]"), 					AddUserAttribute<TArray<FIntPoint>>},
+		{TEXT("VtArray<GfVec2i>"), 			AddUserAttribute<TArray<FIntPoint>>},
+
+		{TEXT("half3[]"), 					AddUserAttribute<TArray<FVector3h>, TArray<FVector3f>>},
+		{TEXT("VtArray<GfVec3h>"), 			AddUserAttribute<TArray<FVector3h>, TArray<FVector3f>>},
+		{TEXT("point3h[]"), 				AddUserAttribute<TArray<FVector3h>, TArray<FVector3f>>},
+		{TEXT("normal3h[]"), 				AddUserAttribute<TArray<FVector3h>, TArray<FVector3f>>},
+		{TEXT("vector3h[]"), 				AddUserAttribute<TArray<FVector3h>, TArray<FVector3f>>},
+		{TEXT("color3h[]"), 				AddUserAttribute<TArray<FVector3h>, TArray<FLinearColor>>},
+
+		{TEXT("float3[]"), 					AddUserAttribute<TArray<FVector3f>>},
+		{TEXT("VtArray<GfVec3f>"), 			AddUserAttribute<TArray<FVector3f>>},
+		{TEXT("point3f[]"), 				AddUserAttribute<TArray<FVector3f>>},
+		{TEXT("normal3f[]"), 				AddUserAttribute<TArray<FVector3f>>},
+		{TEXT("vector3f[]"), 				AddUserAttribute<TArray<FVector3f>>},
+		{TEXT("color3f[]"), 				AddUserAttribute<TArray<FVector3f>, TArray<FLinearColor>>},
+
+		{TEXT("double3[]"), 				AddUserAttribute<TArray<FVector3d>>},
+		{TEXT("VtArray<GfVec3d>"), 			AddUserAttribute<TArray<FVector3d>>},
+		{TEXT("point3d[]"), 				AddUserAttribute<TArray<FVector3d>>},
+		{TEXT("normal3d[]"), 				AddUserAttribute<TArray<FVector3d>>},
+		{TEXT("vector3d[]"), 				AddUserAttribute<TArray<FVector3d>>},
+		{TEXT("color3d[]"), 				AddUserAttribute<TArray<FVector3d>, TArray<FLinearColor>>},
+
+		{TEXT("int3[]"), 					AddUserAttribute<TArray<FIntVector>>},
+		{TEXT("VtArray<GfVec3i>"), 			AddUserAttribute<TArray<FIntVector>>},
+
+		{TEXT("half4[]"), 					AddUserAttribute<TArray<FVector4h>, TArray<FVector4f>>},
+		{TEXT("VtArray<GfVec4h>"), 			AddUserAttribute<TArray<FVector4h>, TArray<FVector4f>>},
+		{TEXT("float4[]"), 					AddUserAttribute<TArray<FVector4f>>},
+		{TEXT("VtArray<GfVec4f>"), 			AddUserAttribute<TArray<FVector4f>>},
+		{TEXT("double4[]"), 				AddUserAttribute<TArray<FVector4d>>},
+		{TEXT("VtArray<GfVec4d>"), 			AddUserAttribute<TArray<FVector4d>>},
+		{TEXT("int4[]"), 					AddUserAttribute<TArray<FIntRect>>},
+		{TEXT("VtArray<GfVec4i>"), 			AddUserAttribute<TArray<FIntRect>>},
+		{TEXT("color4h[]"), 				AddUserAttribute<TArray<FVector4h>, TArray<FLinearColor>>},
+		{TEXT("color4f[]"), 				AddUserAttribute<TArray<FVector4f>, TArray<FLinearColor>>},
+		{TEXT("color4d[]"), 				AddUserAttribute<TArray<FVector4d>, TArray<FLinearColor>>},
 	};
 
 	void TranslateAttributes(const UE::FUsdPrim& Prim, UInterchangeSceneNode* Node, const FString& AllowedAttributeRegex)
@@ -3854,84 +4092,56 @@ namespace UE::InterchangeUsdTranslator::Private
 				continue;
 			}
 
-			// Note: We could have used Cpp type names here instead of the value type token, but doing the
-			// latter lets us handle the color types as FLinearColor, which is probably more useful
-			// References:
-			// - https://openusd.org/docs/api/_usd__page__datatypes.html
-			// - EAttributeTypes declaration
-			// - Engine/Source/Runtime/Interchange/Core/Private/Tests/StorageTest.cpp
-			// clang-format off
-			using TranslationFunc = TFunction<void(const UE::FUsdAttribute&, const FString&, UInterchangeSceneNode*)>;
-			static TMap<FName, TranslationFunc> TranslationFuncs = {
-				{TEXT("bool"),  		TranslateAttribute<bool>},
-
-				{TEXT("uchar"), 		TranslateAttribute<uint8>},
-				{TEXT("uchar[]"), 		TranslateAttribute<TArray<uint8>>},
-				{TEXT("int"), 			TranslateAttribute<int32>},
-				{TEXT("uint"), 			TranslateAttribute<uint32>},
-				{TEXT("int64"), 		TranslateAttribute<int64>},
-				{TEXT("uint64"), 		TranslateAttribute<uint64>},
-
-				{TEXT("half"),	 		TranslateAttribute<FFloat16>},
-				{TEXT("float"), 		TranslateAttribute<float>},
-				{TEXT("double"), 		TranslateAttribute<double>},
-				{TEXT("timecode"), 		TranslateAttribute<double>},
-
-				{TEXT("string"), 		TranslateAttribute<FString>},
-				{TEXT("token"), 		TranslateAttribute<FString>},
-				{TEXT("asset"), 		TranslateAttribute<FString>},
-
-				{TEXT("matrix2d"), 		TranslateAttribute<FMatrix44d>}, // Interchange only supports 4x4 matrices
-				{TEXT("matrix3d"), 		TranslateAttribute<FMatrix44d>},
-				{TEXT("matrix4d"), 		TranslateAttribute<FMatrix44d>},
-				{TEXT("frame4d"), 		TranslateAttribute<FMatrix44d>},
-
-				{TEXT("quath"), 		TranslateAttribute<FQuat4f>}, // No analogue type for Quat of halfs
-				{TEXT("quatf"), 		TranslateAttribute<FQuat4f>},
-				{TEXT("quatd"), 		TranslateAttribute<FQuat4d>},
-
-				{TEXT("half2"), 		TranslateAttribute<FVector2DHalf>},
-				{TEXT("float2"), 		TranslateAttribute<FVector2f>},
-				{TEXT("double2"), 		TranslateAttribute<FVector2d>},
-				{TEXT("int2"), 			TranslateAttribute<FIntPoint>},
-
-				{TEXT("half3"), 		TranslateAttribute<FVector3f>}, // There is no FVector3DHalf
-				{TEXT("point3h"), 		TranslateAttribute<FVector3f>},
-				{TEXT("normal3h"), 		TranslateAttribute<FVector3f>},
-				{TEXT("vector3h"), 		TranslateAttribute<FVector3f>},
-				{TEXT("color3h"), 		TranslateAttribute<FLinearColor>},
-
-				{TEXT("float3"), 		TranslateAttribute<FVector3f>},
-				{TEXT("point3f"), 		TranslateAttribute<FVector3f>},
-				{TEXT("normal3f"), 		TranslateAttribute<FVector3f>},
-				{TEXT("vector3f"), 		TranslateAttribute<FVector3f>},
-				{TEXT("color3f"), 		TranslateAttribute<FLinearColor>},
-
-				{TEXT("double3"), 		TranslateAttribute<FVector3d>},
-				{TEXT("point3d"), 		TranslateAttribute<FVector3d>},
-				{TEXT("normal3d"), 		TranslateAttribute<FVector3d>},
-				{TEXT("vector3d"), 		TranslateAttribute<FVector3d>},
-				{TEXT("color3d"), 		TranslateAttribute<FLinearColor>},
-
-				{TEXT("int3"), 			TranslateAttribute<FIntVector>},
-
-				{TEXT("half4"), 		TranslateAttribute<FVector4f>}, // There is no FVector4DHalf
-				{TEXT("float4"), 		TranslateAttribute<FVector4f>},
-				{TEXT("double4"), 		TranslateAttribute<FVector4d>},
-				{TEXT("int4"), 			TranslateAttribute<FIntRect>},
-				{TEXT("color4h"), 		TranslateAttribute<FLinearColor>},
-				{TEXT("color4f"), 		TranslateAttribute<FLinearColor>},
-				{TEXT("color4d"), 		TranslateAttribute<FLinearColor>},
-			};
-			// clang-format on
-
-			if (TranslationFunc Func = TranslationFuncs.FindRef(Attr.GetTypeName()))
+			if (AddUserAttributeFunc Func = AddUserAttributeFuncs.FindRef(Attr.GetTypeName()))
 			{
-				Func(Attr, AttrName, Node);
+				UE::FVtValue Value;
+				if (Attr.Get(Value)) // Always check for an opinion on the default time code
+				{
+					Func(AttrName, Value, Node);
+				}
 			}
 		}
 	}
 
+	void TranslateMetadata(const UE::FUsdPrim& Prim, UInterchangeSceneNode* Node, const FString& AllowedMetadataRegex)
+	{
+		TRACE_CPUPROFILER_EVENT_SCOPE(TranslateMetadata)
+
+		FRegexPattern RegexPattern{AllowedMetadataRegex};
+
+		TFunction<void(const TMap<FString, UE::FVtValue>&, const FString&)> ParseMetadataMap;
+		ParseMetadataMap = [&ParseMetadataMap, &RegexPattern, Node](const TMap<FString, UE::FVtValue>& MetadataMap, const FString& ConcatParentKey)
+		{
+			for (const TPair<FString, UE::FVtValue>& Pair : MetadataMap)
+			{
+				const FString FullKey = ConcatParentKey.IsEmpty() ? Pair.Key : FString::Printf(TEXT("%s:%s"), *ConcatParentKey, *Pair.Key);
+				const FString TypeName = Pair.Value.GetTypeName();
+
+				if (TypeName == TEXT("VtDictionary"))
+				{
+					TMap<FString, UE::FVtValue> NestedMap = Pair.Value.Get<TMap<FString, UE::FVtValue>>();
+					ParseMetadataMap(NestedMap, FullKey);
+				}
+				else if (AddUserAttributeFunc Func = AddUserAttributeFuncs.FindRef(*TypeName))
+				{
+					// Only match when parsing leaf keys, because a key may not match just the parent nested map key directly and
+					// would have otherwise prevented us from stepping into it (e.g. the regex "customData:my" will not match 
+					// "customData", but it would match "customData:myNestedMap")
+					FRegexMatcher RegexMatcher{RegexPattern, FullKey};
+					if (!RegexMatcher.FindNext())
+					{
+						continue;
+					}
+
+					Func(FullKey, Pair.Value, Node);
+				}
+			}
+		};
+
+		const FString ConcatParentKey;
+		ParseMetadataMap(Prim.GetAllAuthoredMetadata(), ConcatParentKey);
+	}
+
 	void Traverse(
 		const UE::FUsdPrim& Prim,
 		UInterchangeUSDTranslatorImpl& TranslatorImpl,
@@ -4042,6 +4252,11 @@ namespace UE::InterchangeUsdTranslator::Private
 				TranslateAttributes(Prim, SceneNode, TranslatorSettings->AttributeRegexFilter);
 			}
 
+			if (TranslatorSettings->bTranslatePrimMetadata)
+			{
+				TranslateMetadata(Prim, SceneNode, TranslatorSettings->MetadataRegexFilter);
+			}
+
 			// If we're an Xformable, get our transform.
 			// All SceneNodes should have their LocalTransform set though.
 			// Not setting will cause ensure hits in Skeleton generations for example.
@@ -4914,6 +5129,8 @@ UInterchangeUsdTranslatorSettings::UInterchangeUsdTranslatorSettings()
 	  }
 	, bTranslatePrimAttributes(false) // False by default as it could be expensive to traverse all attributes of all prims running the regex
 	, AttributeRegexFilter(TEXT("."))
+	, bTranslatePrimMetadata(false)
+	, MetadataRegexFilter(TEXT("."))
 {
 }
 
diff --git a/Engine/Plugins/Interchange/Extensions/OpenUSD/Source/Import/Public/InterchangeUsdTranslator.h b/Engine/Plugins/Interchange/Extensions/OpenUSD/Source/Import/Public/InterchangeUsdTranslator.h
--- a/Engine/Plugins/Interchange/Extensions/OpenUSD/Source/Import/Public/InterchangeUsdTranslator.h
+++ b/Engine/Plugins/Interchange/Extensions/OpenUSD/Source/Import/Public/InterchangeUsdTranslator.h
@@ -54,19 +54,37 @@ public:
 	UPROPERTY(EditAnywhere, Category = "USD Translator", meta = (EditCondition = bOverrideStageOptions))
 	FUsdStageOptions StageOptions;
 
-	/** Whether to convert USD prim attributes into translated node custom attributes, if they pass the regex filter */
+	/**
+	 * Whether to convert USD prim attributes into translated node user attributes, ultimately ending up imported as asset and actor metadata.
+	 *
+	 * Note that AttributeRegexFilter may be needed in order to select only the desired metadata, as otherwise regular attributes like 
+	 * 'points' and 'faceVertexCounts' will also get translated into Interchange metadata, which could be expensive.
+	 */
 	UPROPERTY(EditAnywhere, Category = "USD Translator")
 	bool bTranslatePrimAttributes;
 
-	/** Regex filter to select which USD prim attributes should be converted into translated node custom attributes */
+	/** Regex to match against USD prim attribute names to decide whether to convert them into translated node user attributes */
 	UPROPERTY(EditAnywhere, Category = "USD Translator", meta = (EditCondition = bTranslatePrimAttributes))
 	FString AttributeRegexFilter;
 
+	/** Whether to convert USD prim metadata into translated node user attributes, if they pass the regex filter. See MetadataRegexFilter. */
+	UPROPERTY(EditAnywhere, Category = "USD Translator")
+	bool bTranslatePrimMetadata;
+
+	/**
+	 * Regex to match against the full USD prim metadata key to decide whether to convert them into translated node user attributes.
+	 *
+	 * Note that nested dictionaries are flattened using the ":" character as separator. So a myMetadata key inside the customData
+	 * dictionary will end up with the full key 'customData:myMetadata', and could be selected by a regex filter 'customData:my',
+	 * or just 'custom', or 'my', etc.
+	 */
+	UPROPERTY(EditAnywhere, Category = "USD Translator", meta = (EditCondition = bTranslatePrimMetadata))
+	FString MetadataRegexFilter;
+
 public:
 	UInterchangeUsdTranslatorSettings();
 };
 
-/* For now, USD Interchange (FBX parity) translator supports textures, materials and static meshes */
 UCLASS(BlueprintType)
 class UInterchangeUSDTranslator
 	: public UInterchangeTranslatorBase
diff --git a/Engine/Plugins/Runtime/USDCore/Source/UnrealUSDWrapper/Private/UsdWrappers/UsdPrim.cpp b/Engine/Plugins/Runtime/USDCore/Source/UnrealUSDWrapper/Private/UsdWrappers/UsdPrim.cpp
--- a/Engine/Plugins/Runtime/USDCore/Source/UnrealUSDWrapper/Private/UsdWrappers/UsdPrim.cpp
+++ b/Engine/Plugins/Runtime/USDCore/Source/UnrealUSDWrapper/Private/UsdWrappers/UsdPrim.cpp
@@ -195,6 +195,43 @@ namespace UE
 #endif	  // #if USE_USD_SDK
 	}
 
+	TMap<FString, UE::FVtValue> FUsdPrim::GetAllMetadata() const
+	{
+		// Note that in USD the key types are pxr::TfToken, and so traditionally we should be using FNames instead of FStrings here.
+		// That could cause trouble however, due to how pxr::TfToken is case sensitive and FNames are "case preserving" (sometimes).
+		// Also, it's a bit annoying how pxr::VtDictionary (in case of nested metadata) does have std::string as keys... So for overall
+		// convenience we use FString here so that we can use the same types for both this result and nested metadata maps, and can
+		// avoid issues due to casing differences
+		TMap<FString, UE::FVtValue> Result;
+#if USE_USD_SDK
+		std::map<pxr::TfToken, pxr::VtValue, pxr::TfDictionaryLessThan> UsdMetadata = Impl->PxrUsdPrim.Get().GetAllMetadata();
+		Result.Reserve(UsdMetadata.size());
+
+		std::map<pxr::TfToken, pxr::VtValue, pxr::TfDictionaryLessThan>::iterator Iter = UsdMetadata.begin();
+		for (; Iter != UsdMetadata.end(); ++Iter)
+		{
+			Result.Add(UTF8_TO_TCHAR(Iter->first.GetString().c_str()), UE::FVtValue{std::move(Iter->second)});
+		}
+#endif	  // #if USE_USD_SDK
+		return Result;
+	}
+
+	TMap<FString, UE::FVtValue> FUsdPrim::GetAllAuthoredMetadata() const
+	{
+		TMap<FString, UE::FVtValue> Result;
+#if USE_USD_SDK
+		std::map<pxr::TfToken, pxr::VtValue, pxr::TfDictionaryLessThan> UsdMetadata = Impl->PxrUsdPrim.Get().GetAllAuthoredMetadata();
+		Result.Reserve(UsdMetadata.size());
+
+		std::map<pxr::TfToken, pxr::VtValue, pxr::TfDictionaryLessThan>::iterator Iter = UsdMetadata.begin();
+		for (; Iter != UsdMetadata.end(); ++Iter)
+		{
+			Result.Add(UTF8_TO_TCHAR(Iter->first.GetString().c_str()), UE::FVtValue{std::move(Iter->second)});
+		}
+#endif	  // #if USE_USD_SDK
+		return Result;
+	}
+
 	bool FUsdPrim::SetSpecifier(ESdfSpecifier Specifier)
 	{
 #if USE_USD_SDK
diff --git a/Engine/Plugins/Runtime/USDCore/Source/UnrealUSDWrapper/Public/UsdWrappers/UsdPrim.h b/Engine/Plugins/Runtime/USDCore/Source/UnrealUSDWrapper/Public/UsdWrappers/UsdPrim.h
--- a/Engine/Plugins/Runtime/USDCore/Source/UnrealUSDWrapper/Public/UsdWrappers/UsdPrim.h
+++ b/Engine/Plugins/Runtime/USDCore/Source/UnrealUSDWrapper/Public/UsdWrappers/UsdPrim.h
@@ -85,6 +85,8 @@ namespace UE
 		bool HasMetadata(const TCHAR* Key) const;
 		bool SetMetadata(const TCHAR* Key, const UE::FVtValue& Value) const;
 		bool ClearMetadata(const TCHAR* Key) const;
+		TMap<FString, UE::FVtValue> GetAllMetadata() const;
+		TMap<FString, UE::FVtValue> GetAllAuthoredMetadata() const;
 
 		// Wrapped pxr::UsdPrim functions, refer to the USD SDK documentation
 	public:
