[Interchange-USD] Fine tuning PointInstancer handling (optimizing in the Translator).
Removed ISM Collapsing mechanism from Interchange Generic Level pipeline.
Added ISM Collapsed traversal to USD Translator.

#rnx

diff --git a/Engine/Plugins/Interchange/Extensions/OpenUSD/Source/Import/Private/InterchangeUsdTranslator.cpp b/Engine/Plugins/Interchange/Extensions/OpenUSD/Source/Import/Private/InterchangeUsdTranslator.cpp
--- a/Engine/Plugins/Interchange/Extensions/OpenUSD/Source/Import/Private/InterchangeUsdTranslator.cpp
+++ b/Engine/Plugins/Interchange/Extensions/OpenUSD/Source/Import/Private/InterchangeUsdTranslator.cpp
@@ -121,13 +121,10 @@ namespace UE::InterchangeUsdTranslator::Private
 	const static FString PrimitiveShapePrefix = TEXT("\\PrimitiveShape\\");
 	const static FString TexturePrefix = TEXT("\\Texture\\");
 	const static FString VolumePrefix = TEXT("\\Volume\\");
-	const static FString InstancerPrefix = TEXT("\\Instancer\\");
-	const static FString InstancePrefix = TEXT("\\Instance\\");
 	const static FString ComponentPrefix = TEXT("\\Component\\");
 
 	const static FString LODContainerSuffix = TEXT("LODContainer");
 	const static FString TwoSidedSuffix = TEXT("_TwoSided");
-	const static FString InstanceSuffix = TEXT("_Instance");
 
 	const static FString LODString = TEXT("LOD");
 
@@ -244,38 +241,6 @@ namespace UE::InterchangeUsdTranslator::Private
 		FTraversalInfo Info
 	);
 
-	struct FUidAndNameCounter
-	{
-		TMap<FString, int32> UidCounter;
-		TMap<FString, int32> NameCounter;
-		
-		void UpdateUid(FString& Uid)
-		{
-			int32& Counter = UidCounter.FindOrAdd(Uid);
-			if (Counter)
-			{
-				Uid = Uid + TEXT("-") + FString::FromInt(Counter);
-			}
-			Counter++;
-		}
-
-		void UpdateName(FString& Name)
-		{
-			int32& Counter = NameCounter.FindOrAdd(Name);
-			if (Counter)
-			{
-				Name = Name + TEXT("-") + FString::FromInt(Counter);
-			}
-			Counter++;
-		}
-
-		void Update(FString& Uid, FString& Name)
-		{
-			UpdateUid(Uid);
-			UpdateName(Name);
-		}
-	};
-
 	// clang-format off
 	const static TMap<FName, EInterchangePropertyTracks> PropertyNameToTrackType = {
 		// Common properties
@@ -447,9 +412,6 @@ namespace UE::InterchangeUsdTranslator::Private
 		// We stash the info we collected from each Volume prim path here, as we'll reuse it between translation and retrieving the payloads
 		TMap<FString, TArray<UsdUtils::FVolumePrimInfo>> PrimPathToVolumeInfo;
 
-		//To be able to create unique Instance specific Uids and Names. Specifically for PointInstancer support:
-		FUidAndNameCounter UidAndNameCounter;
-
 	private:
 		TMap<FString, TArray<FUsdMaterialXShaderGraph::FGeomProp>> MaterialUidToGeomProps;
 
@@ -1422,53 +1384,60 @@ namespace UE::InterchangeUsdTranslator::Private
 		const FString* BaseParameterName = nullptr;
 	};
 
-	//Helper to traverse PointInstancer subhierarchy and process it into UInterchangeSceneComponents and InstancedStaticMeshComponent.
+	//Helper to traverse PointInstancer subhierarchy and process it into InstancedStaticMeshComponents (only).
 	struct FISMComponentTraversalHelper
 	{
 		UInterchangeUSDTranslatorImpl& TranslatorImpl;
 		UInterchangeBaseNodeContainer& NodeContainer;
 		UInterchangeUsdTranslatorSettings* TranslatorSettings;
+		UInterchangeSceneNode* ParentSceneNode;
+		const FString ParentSceneNodeUid;
+
+		//While we could lookup from the NodeContainer, 
+		// we are implementing a LookUp table specifically for ISMComponentNodes for speed:
+		TMap<FString, UInterchangeInstancedStaticMeshComponentNode*> UidToComponentNode;
 
 		FISMComponentTraversalHelper(
 			UInterchangeUSDTranslatorImpl& InTranslatorImpl,
 			UInterchangeBaseNodeContainer& InNodeContainer,
-			UInterchangeUsdTranslatorSettings* InTranslatorSettings)
+			UInterchangeUsdTranslatorSettings* InTranslatorSettings,
+			UInterchangeSceneNode* InParentSceneNode)
 			: TranslatorImpl(InTranslatorImpl)
 			, NodeContainer(InNodeContainer)
 			, TranslatorSettings(InTranslatorSettings)
+			, ParentSceneNode(InParentSceneNode)
+			, ParentSceneNodeUid(ParentSceneNode->GetUniqueID())
 		{
 
 		}
-		
-		UInterchangeSceneComponentNode* CreateSceneComponent(UInterchangeBaseNode* ParentBaseNode, FString& UniqueId, FString& DisplayLabel, const FTransform& LocalTransform, bool bVisible, UClass* ComponentClass = UInterchangeSceneComponentNode::StaticClass())
-		{
-			TranslatorImpl.UidAndNameCounter.Update(UniqueId, DisplayLabel);
 
-			UInterchangeSceneComponentNode* SceneComponentNode = NewObject<UInterchangeSceneComponentNode>(&NodeContainer, ComponentClass);
-			NodeContainer.SetupNode(SceneComponentNode, UniqueId, DisplayLabel, EInterchangeNodeContainerType::TranslatedScene, ParentBaseNode ? ParentBaseNode->GetUniqueID() : TEXT(""));
+		UInterchangeInstancedStaticMeshComponentNode* CreateSceneComponent(const FString& MeshUid)
+		{
+			FString UniqueId = ComponentPrefix + TEXT("\\") + ParentSceneNodeUid + TEXT("\\") + MeshUid;
 
-			SceneComponentNode->SetCustomLocalTransform(LocalTransform);
-			SceneComponentNode->SetCustomComponentVisibility(bVisible);
+			UInterchangeInstancedStaticMeshComponentNode* ISMComponentNode = UidToComponentNode.FindRef(UniqueId);
 
-			if (UInterchangeSceneComponentNode* ParentSceneComponentNode = Cast<UInterchangeSceneComponentNode>(ParentBaseNode))
+			if (ISMComponentNode)
 			{
-				ParentSceneComponentNode->AddComponentUid(SceneComponentNode->GetUniqueID());
+				return ISMComponentNode;
 			}
-			else if (UInterchangeSceneNode* ParentSceneNode = Cast<UInterchangeSceneNode>(ParentBaseNode))
+
+			const UInterchangeBaseNode* MeshNodeOrLODContainer = NodeContainer.GetNode(MeshUid);
+			if (!MeshNodeOrLODContainer)
 			{
-				ParentSceneNode->AddComponentUid(SceneComponentNode->GetUniqueID());
+				return nullptr;
 			}
 
-			return SceneComponentNode;
-		}
-		
-		UInterchangeSceneComponentNode* CreateInstanceSceneComponent(UInterchangeBaseNode* ParentBaseNode, FString& UniqueId, FString& DisplayLabel, int32 InstanceCounter, const FTransform& LocalTransform, bool bVisible)
-		{
-			FString InstanceUid = UniqueId + InstancePrefix + FString::FromInt(InstanceCounter);
-			FString InstanceDisplayLabel = DisplayLabel + InstanceSuffix + FString::FromInt(InstanceCounter);
+			ISMComponentNode = NewObject<UInterchangeInstancedStaticMeshComponentNode>(&NodeContainer);
+			FString DisplayLabel = MeshNodeOrLODContainer->GetDisplayLabel();
+
+			NodeContainer.SetupNode(ISMComponentNode, UniqueId, DisplayLabel, EInterchangeNodeContainerType::TranslatedScene, ParentSceneNode->GetUniqueID());
 
-			//Note: For handling the point instancer's instance's transform we create a component here. This is an extra introduced component compared to the original hierarchy.
-			return CreateSceneComponent(ParentBaseNode, InstanceUid, InstanceDisplayLabel, LocalTransform, bVisible);
+			ParentSceneNode->AddComponentUid(ISMComponentNode->GetUniqueID());
+
+			UidToComponentNode.Add(UniqueId, ISMComponentNode);
+
+			return ISMComponentNode;
 		}
 
 		FString ConvertLodContainerPrim(const UE::FUsdPrim& Prim, FString& ActiveVariantUSDPrimPath, const FTraversalInfo& Info)
@@ -1559,7 +1528,7 @@ namespace UE::InterchangeUsdTranslator::Private
 		void TraverseNestedPointInstancerHierarchy(
 			const UE::FUsdPrim& Prim,
 			FTraversalInfo Info,
-			UInterchangeSceneComponentNode* ParentSceneComponentNode
+			const FTransform& ParentGlobalTransform
 		)
 		{
 			const FUsdStageInfo StageInfo(Prim.GetStage());
@@ -1568,50 +1537,38 @@ namespace UE::InterchangeUsdTranslator::Private
 			bool bResetTransformStack = false;
 			UsdToUnreal::ConvertXformable(Prim.GetStage(), UE::FUsdTyped(Prim), xFormTransform, UsdUtils::GetEarliestTimeCode(), &bResetTransformStack);
 
-			FString ComponentUid = ComponentPrefix + Prim.GetPrimPath().GetString();
-			FString ComponentDisplayLabel(Prim.GetName().ToString());
-
 			Info.UpdateWithCurrentPrim(Prim);
 
+			FTransform CurrentGlobalTransform = xFormTransform * ParentGlobalTransform;
+
 			if (Prim.IsA(TEXT("PointInstancer")))
 			{
-				UInterchangeSceneComponentNode* SceneComponentNode = CreateSceneComponent(ParentSceneComponentNode, ComponentUid, ComponentDisplayLabel, xFormTransform, Info.bVisible);
-
-				if (SceneComponentNode)
-				{
-					TraversePointInstancerHierarchy(Prim, Info, SceneComponentNode);
-				}
+				TraversePointInstancerHierarchy(Prim, Info, CurrentGlobalTransform);
 			}
 			else if (Info.bIsLODContainer)
 			{
 				FString ActiveVariantUSDPrimPath;
 				FString LODContainerUid = ConvertLodContainerPrim(Prim, ActiveVariantUSDPrimPath, Info);
 
-				//This could be a basic StaticMesh, but currently we only have support for InstancedStaticMeshes. (within the context of Components)
-				UInterchangeInstancedStaticMeshComponentNode* ISMComponentNode = Cast<UInterchangeInstancedStaticMeshComponentNode>(CreateSceneComponent(ParentSceneComponentNode, ComponentUid, ComponentDisplayLabel, FTransform::Identity, Info.bVisible, UInterchangeInstancedStaticMeshComponentNode::StaticClass()));
-
-				if (!ISMComponentNode)
+				if (Info.bVisible)
 				{
-					return;
-				}
+					UInterchangeInstancedStaticMeshComponentNode* ISMComponentNode = CreateSceneComponent(LODContainerUid);
 
-				ISMComponentNode->AddInstanceTransform(xFormTransform);
-				if (LODContainerUid.Len())
-				{
-					ISMComponentNode->SetCustomInstancedAssetUid(LODContainerUid);
+					if (!ISMComponentNode)
+					{
+						return;
+					}
+
+					ISMComponentNode->AddInstanceTransform(CurrentGlobalTransform);
+					if (LODContainerUid.Len())
+					{
+						ISMComponentNode->SetCustomInstancedAssetUid(LODContainerUid);
+					}
 				}
 
 				TArray<FUsdPrim> Children = Prim.GetChildren();
 				if (Children.Num())
 				{
-					//Note: For handling the point instancer's instance's transform we create a component here. This is an extra introduced component compared to the original hierarchy.
-					UInterchangeSceneComponentNode* SceneComponentNode = CreateInstanceSceneComponent(ISMComponentNode, ComponentUid, ComponentDisplayLabel, 0 /*InstanceCounter*/, xFormTransform, Info.bVisible);
-
-					if (!SceneComponentNode)
-					{
-						return;
-					}
-
 					for (FUsdPrim& Child : Children)
 					{
 						if (ActiveVariantUSDPrimPath == Child.GetPrimPath().GetString())
@@ -1619,59 +1576,47 @@ namespace UE::InterchangeUsdTranslator::Private
 							TArray<FUsdPrim> ActiveVariantChildren = Child.GetChildren();
 							for (const FUsdPrim& ActiveVariantChild : ActiveVariantChildren)
 							{
-								TraverseNestedPointInstancerHierarchy(ActiveVariantChild, Info, SceneComponentNode);
+								TraverseNestedPointInstancerHierarchy(ActiveVariantChild, Info, CurrentGlobalTransform);
 							}
 						}
 						else
 						{
-							TraverseNestedPointInstancerHierarchy(Child, Info, SceneComponentNode);
+							TraverseNestedPointInstancerHierarchy(Child, Info, CurrentGlobalTransform);
 						}
 					}
 				}
 			}
 			else if (Prim.IsA(TEXT("Xform")))
 			{
-				UInterchangeSceneComponentNode* SceneComponentNode = CreateSceneComponent(ParentSceneComponentNode, ComponentUid, ComponentDisplayLabel, xFormTransform, Info.bVisible, UInterchangeSceneComponentNode::StaticClass());
-
-				if (!SceneComponentNode)
-				{
-					return;
-				}
-
 				TArray<FUsdPrim> Children = Prim.GetChildren();
 				for (const FUsdPrim& Child : Children)
 				{
-					TraverseNestedPointInstancerHierarchy(Child, Info, SceneComponentNode);
+					TraverseNestedPointInstancerHierarchy(Child, Info, CurrentGlobalTransform);
 				}
 			}
 			else if (Prim.IsA(TEXT("Mesh")))
 			{
-				//This could be a basic StaticMesh, but currently we only have support for InstancedStaticMeshes. (within the context of Components)
-				UInterchangeInstancedStaticMeshComponentNode* ISMComponentNode = Cast<UInterchangeInstancedStaticMeshComponentNode>(CreateSceneComponent(ParentSceneComponentNode, ComponentUid, ComponentDisplayLabel, FTransform::Identity, Info.bVisible, UInterchangeInstancedStaticMeshComponentNode::StaticClass()));
-
-				if (!ISMComponentNode)
+				if (Info.bVisible)
 				{
-					return;
-				}
+					FString MeshUid = TranslatorImpl.AddMeshNode(Prim, TranslatorSettings, NodeContainer, Info);
+					//This could be a basic StaticMesh, but currently we only have support for InstancedStaticMeshes. (within the context of Components)
+					UInterchangeInstancedStaticMeshComponentNode* ISMComponentNode = Cast<UInterchangeInstancedStaticMeshComponentNode>(CreateSceneComponent(MeshUid));
 
-				ISMComponentNode->AddInstanceTransform(xFormTransform);
-				ISMComponentNode->SetCustomInstancedAssetUid(TranslatorImpl.AddMeshNode(Prim, TranslatorSettings, NodeContainer, Info));
-			}
-			else
-			{
-				UE_LOG(
-					LogUsd,
-					Warning,
-					TEXT("Unexpected Prim type [%s] in Instancer sub components."),
-					*Prim.GetTypeName().ToString()
-				);
+					if (!ISMComponentNode)
+					{
+						return;
+					}
+
+					ISMComponentNode->AddInstanceTransform(CurrentGlobalTransform);
+					ISMComponentNode->SetCustomInstancedAssetUid(MeshUid);
+				}
 			}
 		}
 
 		void TraversePointInstancerHierarchy(
 			const UE::FUsdPrim& Prim,
 			FTraversalInfo Info,
-			UInterchangeBaseNode* ParentBaseNode
+			const FTransform& ParentGlobalTransform
 		)
 		{
 			UsdUtils::FUsdPointInstancerInstancesData InstancesData;
@@ -1686,8 +1631,6 @@ namespace UE::InterchangeUsdTranslator::Private
 				FString PrototypePathString = PrototypeInstanceData.PathString;
 
 				UE::FUsdPrim PrototypePrim = Prim.GetStage().GetPrimAtPath(UE::FSdfPath{ *PrototypePathString });
-				FString InstancerUid = InstancerPrefix + PrototypePathString;
-				FString InstancerDisplayLabel(PrototypePrim.GetName().ToString());
 
 				FTraversalInfo PrototypePrimInfo = Info;
 				PrototypePrimInfo.UpdateWithCurrentPrim(PrototypePrim);
@@ -1697,48 +1640,41 @@ namespace UE::InterchangeUsdTranslator::Private
 					for (int32 InstanceCounter = 0; InstanceCounter < PrototypeInstanceTransforms.Num(); ++InstanceCounter)
 					{
 						//Note: For handling the point instancer's instance's transform we create a component here. This is an extra introduced component compared to the original hierarchy.
-						UInterchangeSceneComponentNode* SceneComponentNode = CreateInstanceSceneComponent(ParentBaseNode, InstancerUid, InstancerDisplayLabel, InstanceCounter, PrototypeInstanceTransforms[InstanceCounter], PrototypePrimInfo.bVisible);
+						FTransform CurrentGlobalTransform = PrototypeInstanceTransforms[InstanceCounter] * ParentGlobalTransform;
 
-						if (SceneComponentNode)
-						{
-							TraverseNestedPointInstancerHierarchy(PrototypePrim, PrototypePrimInfo, SceneComponentNode);
-						}
+						TraverseNestedPointInstancerHierarchy(PrototypePrim, PrototypePrimInfo, CurrentGlobalTransform);
 					}
 				}
 				else if (PrototypePrim.IsA(TEXT("Mesh")))
 				{
-					FTransform xFormTransform = FTransform::Identity;
-					bool bResetTransformStack = false;
-					UsdToUnreal::ConvertXformable(PrototypePrim.GetStage(), UE::FUsdTyped(PrototypePrim), xFormTransform, UsdUtils::GetEarliestTimeCode(), &bResetTransformStack);
-
-					for (FTransform& Transform : PrototypeInstanceTransforms)
+					if (PrototypePrimInfo.bVisible)
 					{
-						Transform = xFormTransform * Transform;
-					}
+						FTransform xFormTransform = FTransform::Identity;
+						bool bResetTransformStack = false;
+						UsdToUnreal::ConvertXformable(PrototypePrim.GetStage(), UE::FUsdTyped(PrototypePrim), xFormTransform, UsdUtils::GetEarliestTimeCode(), &bResetTransformStack);
 
-					UInterchangeInstancedStaticMeshComponentNode* ISMComponentNode = Cast<UInterchangeInstancedStaticMeshComponentNode>(CreateSceneComponent(ParentBaseNode, InstancerUid, InstancerDisplayLabel, FTransform::Identity, PrototypePrimInfo.bVisible, UInterchangeInstancedStaticMeshComponentNode::StaticClass()));
+						for (FTransform& Transform : PrototypeInstanceTransforms)
+						{
+							Transform = xFormTransform * Transform * ParentGlobalTransform;
+						}
 
-					if (!ISMComponentNode)
-					{
-						continue;
+						FString MeshUid = TranslatorImpl.AddMeshNode(PrototypePrim, TranslatorSettings, NodeContainer, PrototypePrimInfo);
+						UInterchangeInstancedStaticMeshComponentNode* ISMComponentNode = CreateSceneComponent(MeshUid);
+
+						if (!ISMComponentNode)
+						{
+							continue;
+						}
+
+						ISMComponentNode->AddInstanceTransforms(PrototypeInstanceTransforms);
+						ISMComponentNode->SetCustomInstancedAssetUid(MeshUid);
 					}
-					
-					ISMComponentNode->AddInstanceTransforms(PrototypeInstanceTransforms);
-					ISMComponentNode->SetCustomInstancedAssetUid(TranslatorImpl.AddMeshNode(PrototypePrim, TranslatorSettings, NodeContainer, PrototypePrimInfo));
-				}
-				else
-				{
-					UE_LOG(
-						LogUsd,
-						Warning,
-						TEXT("Unexpected Prim type [%s] in PointInstancer."),
-						*PrototypePrim.GetTypeName().ToString()
-					);
 				}
 			}
 		}
 	};
 
+
 	FString UInterchangeUSDTranslatorImpl::AddMaterialNode(
 		const UE::FUsdPrim& Prim,
 		UInterchangeUsdTranslatorSettings* TranslatorSettings,
@@ -4401,8 +4337,8 @@ namespace UE::InterchangeUsdTranslator::Private
 
 			if (Prim.IsA(TEXT("PointInstancer")))
 			{
-				FISMComponentTraversalHelper ISMComponentTraversalHelper(TranslatorImpl, NodeContainer, TranslatorSettings);
-				ISMComponentTraversalHelper.TraversePointInstancerHierarchy(Prim, Info, SceneNode);
+				FISMComponentTraversalHelper ISMComponentTraversalHelper(TranslatorImpl, NodeContainer, TranslatorSettings, SceneNode);
+				ISMComponentTraversalHelper.TraversePointInstancerHierarchy(Prim, Info, FTransform());
 			}
 			else
 			{
diff --git a/Engine/Plugins/Interchange/Runtime/Source/Pipelines/Private/InterchangeGenericScenesPipeline.cpp b/Engine/Plugins/Interchange/Runtime/Source/Pipelines/Private/InterchangeGenericScenesPipeline.cpp
--- a/Engine/Plugins/Interchange/Runtime/Source/Pipelines/Private/InterchangeGenericScenesPipeline.cpp
+++ b/Engine/Plugins/Interchange/Runtime/Source/Pipelines/Private/InterchangeGenericScenesPipeline.cpp
@@ -273,13 +273,9 @@ void UInterchangeGenericLevelPipeline::ExecutePipeline(UInterchangeBaseNodeConta
 	// Find all translated node we need for this pipeline
 	PopulateSceneNodesPreImport(SceneNodes, SceneComponentNodes, LodGroupFactoryNodeLookUp, SkinnedMeshNodes);
 
-	//If bCollapseISMSceneComponents is true then the UInterchangeInstancedStaticMeshComponentFactoryNode  will be created from ExecuteSceneNodePreImport.
-	if (!bCollapseISMSceneComponents)
+	for (UInterchangeSceneComponentNode* SceneComponentNode : SceneComponentNodes)
 	{
-		for (UInterchangeSceneComponentNode* SceneComponentNode : SceneComponentNodes)
-		{
-			CreateComponentFactoryNode(SceneComponentNode);
-		}
+		CreateComponentFactoryNode(SceneComponentNode);
 	}
 
 #if WITH_EDITORONLY_DATA
@@ -558,81 +554,6 @@ void UInterchangeGenericLevelPipeline::PopulateSceneNodesPreImport(TArray<UInter
 	});
 }
 
-void UInterchangeGenericLevelPipeline::CreateAndSetUpCollapsedISMComponents(UInterchangeActorFactoryNode* ActorFactoryNode, const UInterchangeSceneNode* SceneNode) const
-{
-	TArray<FString> ComponentUids;
-	SceneNode->GetComponentUids(ComponentUids);
-
-	TMap<FString, TArray<FTransform>> InstanceTransformsPerMeshes;
-	CollectCollapsedInstancedStaticMeshTransforms(ComponentUids, InstanceTransformsPerMeshes, FTransform::Identity);
-
-	for (const TPair<FString, TArray<FTransform>>& Entry : InstanceTransformsPerMeshes)
-	{
-		const UInterchangeBaseNode* MeshBaseNode = BaseNodeContainer->GetNode(Entry.Key);
-		
-		if (!MeshBaseNode)
-		{
-			continue;
-		}
-
-		FString DisplayLabel = MeshBaseNode->GetDisplayLabel();
-		FString UniqueId = ActorFactoryNode->GetUniqueID() + TEXT("\\CollapsedISM\\") + Entry.Key;
-		
-		UInterchangeInstancedStaticMeshComponentFactoryNode * ISMComponentFactoryNode = NewObject<UInterchangeInstancedStaticMeshComponentFactoryNode >(BaseNodeContainer, NAME_None);
-		BaseNodeContainer->SetupNode(ISMComponentFactoryNode, UniqueId, DisplayLabel, EInterchangeNodeContainerType::FactoryData, ActorFactoryNode->GetUniqueID());
-		
-		FString MeshFactoryNodeUid = UInterchangeFactoryBaseNode::BuildFactoryNodeUid(Entry.Key);
-		
-		ISMComponentFactoryNode->SetCustomInstancedAssetUid(MeshFactoryNodeUid);
-		ISMComponentFactoryNode->AddInstanceTransforms(Entry.Value);
-		ISMComponentFactoryNode->SetCustomHierarchicalISM(bUseHierarchicalISMComponents);
-
-		ActorFactoryNode->AddFactoryDependencyUid(MeshFactoryNodeUid);
-		ActorFactoryNode->AddComponentUid(UniqueId);
-	}
-}
-
-void UInterchangeGenericLevelPipeline::CollectCollapsedInstancedStaticMeshTransforms(const TArray<FString>& SceneComponentUids, TMap<FString, TArray<FTransform>>& InstanceTransformsPerMeshes, FTransform ParentGlobalTransform) const
-{
-	for (const FString& SceneComponentUid : SceneComponentUids)
-	{
-		if (const UInterchangeSceneComponentNode* SceneComponentNode = Cast<const UInterchangeSceneComponentNode>(BaseNodeContainer->GetNode(SceneComponentUid)))
-		{
-			bool bVisible = true;
-			SceneComponentNode->GetCustomComponentVisibility(bVisible);
-			if (bVisible)
-			{
-				if (const UInterchangeInstancedStaticMeshComponentNode* ISMComponentNode = Cast<const UInterchangeInstancedStaticMeshComponentNode>(SceneComponentNode))
-				{
-					TArray<FTransform> InstanceTransforms;
-					ISMComponentNode->GetInstanceTransforms(InstanceTransforms);
-					for (FTransform& Transform : InstanceTransforms)
-					{
-						Transform = Transform * ParentGlobalTransform;
-					}
-
-					FString MeshUid;
-					ISMComponentNode->GetCustomInstancedAssetUid(MeshUid);
-					TArray<FTransform>& InstanceTransformsPerMesh = InstanceTransformsPerMeshes.FindOrAdd(MeshUid);
-					InstanceTransformsPerMesh.Append(InstanceTransforms);
-				}
-			}
-			
-
-			TArray<FString> SubSceneComponentUids;
-			SceneComponentNode->GetComponentUids(SubSceneComponentUids);
-			if (SubSceneComponentUids.Num())
-			{
-				FTransform LocalTransform;
-				SceneComponentNode->GetCustomLocalTransform(LocalTransform);
-				FTransform GlobalTransform = LocalTransform * ParentGlobalTransform;
-
-				CollectCollapsedInstancedStaticMeshTransforms(SubSceneComponentUids, InstanceTransformsPerMeshes, GlobalTransform);
-			}
-		}
-	}
-}
-
 void UInterchangeGenericLevelPipeline::SetUpComponents(UInterchangeActorFactoryNode* ActorFactoryNode, const UInterchangeSceneNode* SceneNode) const
 {
 	TArray<FString> SceneComponentUids;
@@ -909,14 +830,7 @@ void UInterchangeGenericLevelPipeline::ExecuteSceneNodePreImport(const FTransfor
 		SetUpFactoryNode(ActorFactoryNode, SceneNode, TranslatedAssetNode);
 	}
 
-	if (bCollapseISMSceneComponents)
-	{
-		CreateAndSetUpCollapsedISMComponents(ActorFactoryNode, SceneNode);
-	}
-	else
-	{
-		SetUpComponents(ActorFactoryNode, SceneNode);
-	}
+	SetUpComponents(ActorFactoryNode, SceneNode);
 	
 
 	//Make sure all actor factory nodes and dependencies have the specified strategy
diff --git a/Engine/Plugins/Interchange/Runtime/Source/Pipelines/Public/InterchangeGenericScenesPipeline.h b/Engine/Plugins/Interchange/Runtime/Source/Pipelines/Public/InterchangeGenericScenesPipeline.h
--- a/Engine/Plugins/Interchange/Runtime/Source/Pipelines/Public/InterchangeGenericScenesPipeline.h
+++ b/Engine/Plugins/Interchange/Runtime/Source/Pipelines/Public/InterchangeGenericScenesPipeline.h
@@ -69,10 +69,6 @@ public:
 	UPROPERTY(EditAnywhere, BlueprintReadWrite, Category = "Scene", meta = (SubCategory = "Reimport Assets"))
 	bool bDeleteMissingAssets = false;
 
-	/* If enabled, Normal SceneComponents found under a given Actor will be ignored and their transforms baked into collapsed Instanced Meshes. */
-	UPROPERTY(EditAnywhere, BlueprintReadWrite, Category = "Scene", meta = (SubCategory = "Component properties", DisplayName = "Collapse ISM Scene Components"))
-	bool bCollapseISMSceneComponents = true;
-
 	UPROPERTY(EditAnywhere, BlueprintReadWrite, Category = "Scene", meta = (SubCategory = "Component properties", DisplayName = "Use Hierarchical ISM Components"))
 	bool bUseHierarchicalISMComponents = false;
 
@@ -137,18 +133,6 @@ protected:
 	 */
 	UE_API virtual void SetUpComponents(UInterchangeActorFactoryNode* ActorFactoryNode, const UInterchangeSceneNode* SceneNode) const;
 
-	/**
-	 * Genereates and sets up the collected Collapsed InstancedStaticMesh transforms.
-	 * Note: Collapsing will ignore User Defined Attributes from the nodes being collapsed.
-	 */
-	UE_API virtual void CreateAndSetUpCollapsedISMComponents(UInterchangeActorFactoryNode* ActorFactoryNode, const UInterchangeSceneNode* SceneNode) const;
-
-	/**
-	 * Collects InstancedStaticMeshComponentNodes collapses them to have 1 Component per Mesh with the required baked transforms on the InstanceTransforms.
-	 * Only collects ISM Components marked as Visible/(Not Marked as InVisible).
-	 */
-	UE_API virtual void CollectCollapsedInstancedStaticMeshTransforms(const TArray<FString>& SceneComponentUids, TMap<FString, TArray<FTransform>>& InstanceTransformsPerMeshes, FTransform ParentGlobalTransform) const;
-
 	UE_API virtual void SetUpComponentDependencies(UInterchangeActorFactoryNode* ActorFactoryNode, const TArray<FString>& SceneComponentUids) const;
 	
 	/** Disable this option to not convert Standard(Perspective) to Physical Cameras*/
