[Interchange] Audio - Support for Interchage USD
- Spatial Audio prims are now imported as audio assets.

#rnx

diff --git a/Engine/Plugins/Interchange/Extensions/OpenUSD/Source/Import/Private/InterchangeUsdTranslator.cpp b/Engine/Plugins/Interchange/Extensions/OpenUSD/Source/Import/Private/InterchangeUsdTranslator.cpp
--- a/Engine/Plugins/Interchange/Extensions/OpenUSD/Source/Import/Private/InterchangeUsdTranslator.cpp
+++ b/Engine/Plugins/Interchange/Extensions/OpenUSD/Source/Import/Private/InterchangeUsdTranslator.cpp
@@ -37,6 +37,7 @@
 #include "UsdWrappers/UsdVariantSets.h"
 #include "UsdWrappers/VtValue.h"
 
+#include "InterchangeAudioSoundWaveNode.h"
 #include "InterchangeCameraNode.h"
 #include "InterchangeGroomNode.h"
 #include "InterchangeLightNode.h"
@@ -82,6 +83,7 @@
 #include "pxr/usd/usdGeom/tokens.h"
 #include "pxr/usd/usdLux/tokens.h"
 #include "pxr/usd/usdShade/tokens.h"
+#include "pxr/usd/usdMedia/tokens.h"
 #include "USDIncludesEnd.h"
 #endif	  // USE_USD_SDK
 
@@ -119,6 +121,7 @@ namespace UE::InterchangeUsdTranslator::Private
 	const static FString MeshPrefix = TEXT("\\Mesh\\");
 	const static FString MorphTargetPrefix = TEXT("\\MorphTarget\\");
 	const static FString PrimitiveShapePrefix = TEXT("\\PrimitiveShape\\");
+	const static FString SpatialAudioPrefix = TEXT("\\SpatialAudio\\");
 	const static FString TexturePrefix = TEXT("\\Texture\\");
 	const static FString VolumePrefix = TEXT("\\Volume\\");
 	const static FString ComponentPrefix = TEXT("\\Component\\");
@@ -318,6 +321,11 @@ namespace UE::InterchangeUsdTranslator::Private
 			UInterchangeBaseNodeContainer& NodeContainer
 		);
 
+		FString AddSpatialAudioNode(
+			const UE::FUsdPrim& Prim,
+			UInterchangeBaseNodeContainer& NodeContainer
+		);
+
 		/**
 		 * Search for a MaterialX file embedded in the file, in case there's none, create a ShaderGraph using MaterialX
 		 * @param GeomPropValueNames - Array of the <geompropvalue> node names from the shader graph that has been converted to <image> nodes (used
@@ -3241,7 +3249,7 @@ namespace UE::InterchangeUsdTranslator::Private
 		const FString NodeName(Prim.GetName().ToString());
 
 		// Check if Node already exist with this ID
-		if (const UInterchangeMeshNode* Node = Cast<UInterchangeMeshNode>(NodeContainer.GetNode(NodeUid)))
+		if (const UInterchangeGroomNode* Node = Cast<UInterchangeGroomNode>(NodeContainer.GetNode(NodeUid)))
 		{
 			return NodeUid;
 		}
@@ -3254,6 +3262,38 @@ namespace UE::InterchangeUsdTranslator::Private
 		return NodeUid;
 	}
 
+	FString UInterchangeUSDTranslatorImpl::AddSpatialAudioNode(
+		const UE::FUsdPrim& Prim,
+		UInterchangeBaseNodeContainer& NodeContainer
+	)
+	{
+		TRACE_CPUPROFILER_EVENT_SCOPE(UInterchangeUSDTranslatorImpl::AddSpatialAudioNode)
+
+		const FString PrimPath = UsdUtils::GetPrototypePrimPath(Prim).GetString();
+		const FString NodeUid = SpatialAudioPrefix + PrimPath;
+		const FString NodeName(Prim.GetName().ToString());
+
+		// Check if Node already exist with this ID
+		if (const UInterchangeAudioSoundWaveNode* Node = Cast<UInterchangeAudioSoundWaveNode>(NodeContainer.GetNode(NodeUid)))
+		{
+			return NodeUid;
+		}
+
+		UInterchangeAudioSoundWaveNode* SoundWaveNode = NewObject<UInterchangeAudioSoundWaveNode>(&NodeContainer);
+		NodeContainer.SetupNode(SoundWaveNode, NodeUid, NodeName, EInterchangeNodeContainerType::TranslatedAsset);
+		SoundWaveNode->SetAssetName(NodeName);
+		
+		static const FString FilePathToken = UsdToUnreal::ConvertToken(pxr::UsdMediaTokens->filePath);
+		UE::FUsdAttribute FilePathAttribute = Prim.GetAttribute(*FilePathToken);
+		const FString ResolvedAudioFilepath = UsdUtils::GetResolvedAssetPath(FilePathAttribute);
+		SoundWaveNode->SetPayloadKey(ResolvedAudioFilepath);
+
+		// TODO: @vjavdekar - Get all the attributes from the media spatial audio
+		// and add those to SoundWaveNode.
+
+		return NodeUid;
+	}
+
 	void AddTrackSetNode(UInterchangeUSDTranslatorImpl& Impl, UInterchangeBaseNodeContainer& NodeContainer)
 	{
 		// For now we only want a single track set (i.e. LevelSequence) per stage.
@@ -4176,20 +4216,13 @@ namespace UE::InterchangeUsdTranslator::Private
 		{
 			AssetNodeUids = {TranslatorImpl.AddGroomNode(Prim, NodeContainer)};
 		}
+		else if (Prim.IsA(TEXT("SpatialAudio")))
+		{
+			AssetNodeUids = { TranslatorImpl.AddSpatialAudioNode(Prim, NodeContainer) };
+		}
 		else
 		{
-			const static TSet<FName> KnownUnsupported = {TEXT("SpatialAudio")};
-
-			if (KnownUnsupported.Contains(TypeName))
-			{
-				UInterchangeResultWarning_Generic* Message = TranslatorImpl.ResultsContainer->Add<UInterchangeResultWarning_Generic>();
-				Message->SourceAssetName = TranslatorImpl.UsdStage.GetRootLayer().GetRealPath();
-				Message->Text = FText::Format(
-					LOCTEXT("UnsupportedSchemaType", "Prim '{0}' has schema '{1}', which is not yet supported via USD Interchange."),
-					FText::FromString(SceneNodeUid),
-					FText::FromName(TypeName)
-				);
-			}
+			// Handle the KnownUnsupportedTypes here.
 		}
 
 		// Only prims that require rendering (and have a renderable parent) get a scene node.
diff --git a/Engine/Plugins/Interchange/Extensions/OpenUSD/Source/Import/Public/InterchangeUsdTranslator.h b/Engine/Plugins/Interchange/Extensions/OpenUSD/Source/Import/Public/InterchangeUsdTranslator.h
--- a/Engine/Plugins/Interchange/Extensions/OpenUSD/Source/Import/Public/InterchangeUsdTranslator.h
+++ b/Engine/Plugins/Interchange/Extensions/OpenUSD/Source/Import/Public/InterchangeUsdTranslator.h
@@ -3,6 +3,7 @@
 #pragma once
 
 #include "Animation/InterchangeAnimationPayloadInterface.h"
+#include "Audio/InterchangeAudioPayloadInterface.h"
 #include "Groom/InterchangeGroomPayloadInterface.h"
 #include "InterchangeTranslatorBase.h"
 #include "Mesh/InterchangeMeshPayloadInterface.h"
@@ -95,6 +96,7 @@ class UInterchangeUSDTranslator
 	, public IInterchangeAnimationPayloadInterface
 	, public IInterchangeVolumePayloadInterface
 	, public IInterchangeGroomPayloadInterface
+	, public IInterchangeAudioPayloadInterface
 {
 	GENERATED_BODY()
 
