Interchange USD:
- Assign volume material coming from a  MaterialX shader graph  to the HeterogeousVolumeActor

diff --git a/Engine/Plugins/Interchange/Extensions/OpenUSD/Source/Import/Private/InterchangeUsdTranslator.cpp b/Engine/Plugins/Interchange/Extensions/OpenUSD/Source/Import/Private/InterchangeUsdTranslator.cpp
--- a/Engine/Plugins/Interchange/Extensions/OpenUSD/Source/Import/Private/InterchangeUsdTranslator.cpp
+++ b/Engine/Plugins/Interchange/Extensions/OpenUSD/Source/Import/Private/InterchangeUsdTranslator.cpp
@@ -319,6 +319,7 @@ namespace UE::InterchangeUsdTranslator::Private
 
 		TArray<FString> AddVolumeNodes(
 			const UE::FUsdPrim& InPrim,
+			UInterchangeUsdTranslatorSettings* TranslatorSettings,
 			UInterchangeBaseNodeContainer& InNodeContainer,
 			FString& OutMaterialInstanceUid,
 			bool& bOutNeedsFrameTrack
@@ -1052,13 +1053,27 @@ namespace UE::InterchangeUsdTranslator::Private
 			[&NodeContainer, &PrimName, &Result](const FString&, UInterchangeShaderGraphNode* ShaderGraphNode)
 			{
 				FString NodeUid = ShaderGraphNode->GetUniqueID();
-				if (FPaths::GetBaseFilename(NodeUid) == PrimName)
+				// Because USD allows having the same name between a parent and its child (unlike MaterialX)
+				// When building the shader graph we may end up with a name not totally equal to the USD one (MaterialX will append an index to the material name)
+				// e.g: prim -> mtlx_volume
+				//      mtlx -> mtlx_volume2
+				// see the egg_jade file from [REDACTED_PROFILE_URL]/MaterialEggs
+				FString BaseName = FPaths::GetBaseFilename(NodeUid);
+				if (BaseName == PrimName)
 				{
 					Result = NodeUid;
 					return true;
 				}
 				else
 				{
+					if (BaseName.StartsWith(PrimName, ESearchCase::CaseSensitive))
+					{
+						if (FString Suffix = BaseName.Mid(PrimName.Len()); Suffix.IsNumeric())
+						{
+							Result = NodeUid;
+							return true;
+						}
+					}
 					return false;
 				}
 			}
@@ -2951,6 +2966,7 @@ namespace UE::InterchangeUsdTranslator::Private
 
 	TArray<FString> UInterchangeUSDTranslatorImpl::AddVolumeNodes(
 		const UE::FUsdPrim& InPrim,
+		UInterchangeUsdTranslatorSettings* TranslatorSettings,
 		UInterchangeBaseNodeContainer& InNodeContainer,
 		FString& OutMaterialInstanceUid,
 		bool& bOutNeedsFrameTrack
@@ -3169,6 +3185,45 @@ namespace UE::InterchangeUsdTranslator::Private
 			}
 		}
 
+		// In case of MaterialX context, the volume material will be created from a shader graph
+		if (TranslatorSettings && (TranslatorSettings->RenderContext == UnrealIdentifiers::MaterialXRenderContext))
+		{
+			UsdUtils::FUsdPrimMaterialAssignmentInfo Assignments = 
+				UsdUtils::GetPrimMaterialAssignments(InPrim, UsdUtils::GetDefaultTimeCode(), false, CachedMeshConversionOptions.RenderContext, CachedMeshConversionOptions.MaterialPurpose);
+
+			for (const UsdUtils::FUsdPrimMaterialSlot& Slot : Assignments.Slots)
+			{
+				FString SlotName = Slot.MaterialSource;
+
+				if (Slot.AssignmentType != UsdUtils::EPrimAssignmentType::MaterialPrim)
+				{
+					continue;
+				}
+
+				// MaterialSource here is the material prim path
+				UE::FUsdPrim MaterialPrim = UsdStage.GetPrimAtPath(UE::FSdfPath{ *Slot.MaterialSource });
+
+				// Go through a prototype in case this material is inside an instance proxy, as this needs to match up with
+				// the polygon group name we're going to add to our mesh description.
+				UE::FUsdPrim MaterialPrototypePrim = UsdUtils::GetPrototypePrim(MaterialPrim);
+				if (MaterialPrototypePrim != MaterialPrim)
+				{
+					SlotName = MaterialPrototypePrim.GetPrimPath().GetString();
+				}
+
+				OutMaterialInstanceUid = AddMaterialNode(MaterialPrototypePrim, TranslatorSettings, InNodeContainer, Slot.bMeshIsDoubleSided);
+
+				if (!OutMaterialInstanceUid.IsEmpty())
+				{
+					break;
+				}
+			}
+
+			if (!OutMaterialInstanceUid.IsEmpty())
+			{
+				return VolumeAssetNodeUids.Array();
+			}
+		}
 		// Setup a new material instance node for this volume, which will be used by the Heterogeneous Volume Actor we'll also
 		// spawn for this volume prim.
 		//
@@ -3465,7 +3520,7 @@ namespace UE::InterchangeUsdTranslator::Private
 	void AddVolumeSceneNodeAttributes(
 		const UE::FUsdPrim& Prim,
 		UInterchangeSceneNode* SceneNode,
-		TArray<FString> AssetNodeUids,
+		const TArray<FString>& AssetNodeUids,
 		const FString& VolumeMaterialInstanceUid,
 		bool bNeedsAnimationTrack,
 		UInterchangeUSDTranslatorImpl& Impl,
@@ -4246,7 +4301,7 @@ namespace UE::InterchangeUsdTranslator::Private
 		}
 		else if (Prim.IsA(TEXT("Volume")))
 		{
-			AssetNodeUids = TranslatorImpl.AddVolumeNodes(Prim, NodeContainer, VolumeMaterialInstanceUid, bNeedsVolumeTrack);
+			AssetNodeUids = TranslatorImpl.AddVolumeNodes(Prim, TranslatorSettings, NodeContainer, VolumeMaterialInstanceUid, bNeedsVolumeTrack);
 		}
 		else if ((Prim.IsA(TEXT("BasisCurves")) || Prim.IsA(TEXT("Xform"))) && Prim.HasAPI(TEXT("GroomAPI")))
 		{
diff --git a/Engine/Plugins/Runtime/USDCore/Source/USDUtilities/Private/USDGeomMeshConversion.cpp b/Engine/Plugins/Runtime/USDCore/Source/USDUtilities/Private/USDGeomMeshConversion.cpp
--- a/Engine/Plugins/Runtime/USDCore/Source/USDUtilities/Private/USDGeomMeshConversion.cpp
+++ b/Engine/Plugins/Runtime/USDCore/Source/USDUtilities/Private/USDGeomMeshConversion.cpp
@@ -4169,11 +4169,15 @@ UsdUtils::FUsdPrimMaterialAssignmentInfo UsdUtils::GetPrimMaterialAssignments(
 			return {};
 		}
 
-		// Ignore this material if UsdToUnreal::ConvertMaterial would as well
-		pxr::UsdShadeShader SurfaceShader = ShadeMaterial.ComputeSurfaceSource(RenderContext);
-		if (!SurfaceShader)
+		if (RenderContext.GetString().c_str() != UnrealIdentifiers::MaterialXRenderContext)
 		{
-			return {};
+			// Ignore this material if UsdToUnreal::ConvertMaterial would as well 
+			// in a MaterialX context we can have any type of material though (surface/displacement/volume)
+			pxr::UsdShadeShader SurfaceShader = ShadeMaterial.ComputeSurfaceSource(RenderContext);
+			if (!SurfaceShader)
+			{
+				return {};
+			}
 		}
 
 		pxr::UsdPrim ShadeMaterialPrim = ShadeMaterial.GetPrim();
