Interchange USD:
- Fix meshes ending up having no geometry,in some cases, with the mtlx render context

diff --git a/Engine/Plugins/Interchange/Extensions/OpenUSD/Source/Import/Private/InterchangeUSDPipeline.cpp b/Engine/Plugins/Interchange/Extensions/OpenUSD/Source/Import/Private/InterchangeUSDPipeline.cpp
--- a/Engine/Plugins/Interchange/Extensions/OpenUSD/Source/Import/Private/InterchangeUSDPipeline.cpp
+++ b/Engine/Plugins/Interchange/Extensions/OpenUSD/Source/Import/Private/InterchangeUSDPipeline.cpp
@@ -1164,6 +1164,7 @@ namespace UE::InterchangeUsdPipeline::Private
 				FVertexInstanceAttributesVertexColorsRefGuard VertexColorsGuard{ VertexColors };
 				FVertexInstanceAttributesNormalsRefGuard NormalsGuard{ Normals };
 
+				bool bConverted = false;
 				for (int32 Index = 0; Index < NumberOfPrimvar; ++Index)
 				{
 					UMaterialExpressionTextureSample* TextureSampleExpression = GetMaterialExpressionTextureSample(Index);
@@ -1179,7 +1180,7 @@ namespace UE::InterchangeUsdPipeline::Private
 					FactoryNode->GetBooleanAttribute(USD::Primvar::TangentSpace + FString::FromInt(Index), bTangentSpace);
 
 					MeshDescriptionPrimvar.VertexInstanceAttributes().ForEach(
-					[this, &PrimvarName, &VertexColors, &Normals, TextureSampleExpression, &MeshDescriptionPrimvar, &DynamicMesh, bTangentSpace]
+					[this, &PrimvarName, &VertexColors, &Normals, TextureSampleExpression, &MeshDescriptionPrimvar, &DynamicMesh, bTangentSpace, &bConverted]
 					(const FName AttributeName, auto AttributesConstRef)
 					{
 						if (AttributeName != PrimvarName)
@@ -1261,17 +1262,20 @@ namespace UE::InterchangeUsdPipeline::Private
 						TextureSampleExpression->ConstCoordinate = UVChannel;
 
 						MaterialExpressionTextureSamples.Add(TextureSampleExpression);
+						bConverted = true;
 					});
 				}
 
 				VertexColorsGuard.Reset();
 				NormalsGuard.Reset();
 
-				FDynamicMeshToMeshDescription Converter;
-				// convert back to the mesh description so it's also adding the new UV set
-				// though this conversion unfortunately doesn't preserve the primvars, do we need them after all that?
-				Converter.Convert(&DynamicMesh, *StaticMeshAsset->GetMeshDescription(0));
-
+				if (bConverted)
+				{
+					FDynamicMeshToMeshDescription Converter;
+					// convert back to the mesh description so it's also adding the new UV set
+					// though this conversion unfortunately doesn't preserve the primvars, do we need them after all that?
+					Converter.Convert(&DynamicMesh, *StaticMeshAsset->GetMeshDescription(0));
+				}
 				UpdateMaterials();
 			}
 		}
