USD: Fix Euler flipping issues when importing some scenes with rotation animation.

diff --git a/Engine/Plugins/Interchange/Extensions/OpenUSD/Source/Import/Private/InterchangeUsdTranslator.cpp b/Engine/Plugins/Interchange/Extensions/OpenUSD/Source/Import/Private/InterchangeUsdTranslator.cpp
--- a/Engine/Plugins/Interchange/Extensions/OpenUSD/Source/Import/Private/InterchangeUsdTranslator.cpp
+++ b/Engine/Plugins/Interchange/Extensions/OpenUSD/Source/Import/Private/InterchangeUsdTranslator.cpp
@@ -860,6 +860,12 @@ namespace UE::InterchangeUsdTranslator::Private
 		const ERichCurveInterpMode InterpMode = (UsdStage.GetInterpolationType() == EUsdInterpolationType::Linear)
 													? ERichCurveInterpMode::RCIM_Linear
 													: ERichCurveInterpMode::RCIM_Constant;
+		struct FHeadingHelper
+		{
+			FRotator Heading;
+			bool bHeadingSet = false;
+		};
+		FHeadingHelper HeadingHelper;
 
 		double LastTimeSample = TNumericLimits<double>::Lowest();
 		for (const double UsdTimeSample : UsdTimeSamples)
@@ -882,6 +888,27 @@ namespace UE::InterchangeUsdTranslator::Private
 			FRotator Rotator = UEValue.Rotator();
 			FVector Scale = UEValue.GetScale3D();
 
+			if (HeadingHelper.bHeadingSet)
+			{
+				FMath::WindRelativeAnglesDegrees(HeadingHelper.Heading.Roll, Rotator.Roll);
+				FMath::WindRelativeAnglesDegrees(HeadingHelper.Heading.Pitch, Rotator.Pitch);
+				FMath::WindRelativeAnglesDegrees(HeadingHelper.Heading.Yaw, Rotator.Yaw);
+
+				FRotator OtherChoice = Rotator.GetEquivalentRotator().GetNormalized();
+				float FirstDiff = HeadingHelper.Heading.GetManhattanDistance(Rotator);
+				float SecondDiff = HeadingHelper.Heading.GetManhattanDistance(OtherChoice);
+				if (SecondDiff < FirstDiff)
+				{
+					Rotator = OtherChoice;
+				}
+			}
+			else
+			{
+				HeadingHelper.bHeadingSet = true;
+			}
+
+			HeadingHelper.Heading = Rotator;
+
 			FKeyHandle HandleTransX = TransXCurve.AddKey(FrameTimeSeconds, Location.X);
 			FKeyHandle HandleTransY = TransYCurve.AddKey(FrameTimeSeconds, Location.Y);
 			FKeyHandle HandleTransZ = TransZCurve.AddKey(FrameTimeSeconds, Location.Z);
