Interchange USD: Bake stage metersPerUnit/upAxis into the volume data.

diff --git a/Engine/Plugins/Interchange/Extensions/OpenUSD/Source/Import/Private/InterchangeUsdTranslator.cpp b/Engine/Plugins/Interchange/Extensions/OpenUSD/Source/Import/Private/InterchangeUsdTranslator.cpp
--- a/Engine/Plugins/Interchange/Extensions/OpenUSD/Source/Import/Private/InterchangeUsdTranslator.cpp
+++ b/Engine/Plugins/Interchange/Extensions/OpenUSD/Source/Import/Private/InterchangeUsdTranslator.cpp
@@ -6040,7 +6040,17 @@ TOptional<UE::Interchange::FVolumePayloadData> UInterchangeUSDTranslator::GetVol
 	const IInterchangeVolumePayloadInterface* VolumeInterface = Cast<IInterchangeVolumePayloadInterface>(ExistingTranslator.Get());
 	if (VolumeInterface)
 	{
-		return VolumeInterface->GetVolumePayloadData(PayloadKey);
+		if (TOptional<UE::Interchange::FVolumePayloadData> Result = VolumeInterface->GetVolumePayloadData(PayloadKey))
+		{
+			// Bake our stage coordinate system transform into the volume, just like we do for geometry.
+			// We do this here because OpenVDB doesn't really specify a fixed coordinate system, and so we interpret it as just meaning
+			// "units" in general. This means we must go through a metersPerUnit / upAxis conversion to get it into UE's world space
+			const FUsdStageInfo StageInfo{Impl->UsdStage};
+			FTransform StageInfoTransform = UsdUtils::GetUsdToUESpaceTransform(StageInfo);
+			Result.GetValue().Transform = Result->Transform * StageInfoTransform;
+
+			return Result;
+		}
 	}
 #endif	  // USE_USD_SDK
 
diff --git a/Engine/Plugins/Runtime/USDCore/Source/USDUtilities/Private/USDTypesConversion.cpp b/Engine/Plugins/Runtime/USDCore/Source/USDUtilities/Private/USDTypesConversion.cpp
--- a/Engine/Plugins/Runtime/USDCore/Source/USDUtilities/Private/USDTypesConversion.cpp
+++ b/Engine/Plugins/Runtime/USDCore/Source/USDUtilities/Private/USDTypesConversion.cpp
@@ -121,6 +121,26 @@ namespace UsdUtils
 
 		return TransformInUESpace;
 	}
+
+	FTransform GetUsdToUESpaceTransform(const FUsdStageInfo& StageInfo)
+	{
+		FTransform MetersPerUnitTransform;
+		const double UEMetersPerUnit = 0.01;
+		if (!FMath::IsNearlyEqual(StageInfo.MetersPerUnit, UEMetersPerUnit))
+		{
+			const double ScaleFactor = StageInfo.MetersPerUnit / UEMetersPerUnit;
+			MetersPerUnitTransform.SetScale3D(FVector{ScaleFactor, ScaleFactor, ScaleFactor});
+		}
+
+		FTransform UpAxisTransform;
+		UpAxisTransform.SetScale3D(FVector{1.0, -1.0, 1.0});
+		if (StageInfo.UpAxis == EUsdUpAxis::YAxis)
+		{
+			UpAxisTransform.SetRotation(FQuat::MakeFromEuler(FVector{90, 0, 0}));
+		}
+
+		return MetersPerUnitTransform * UpAxisTransform;
+	}
 }	 // namespace UsdUtils
 
 namespace UsdToUnreal
diff --git a/Engine/Plugins/Runtime/USDCore/Source/USDUtilities/Public/USDTypesConversion.h b/Engine/Plugins/Runtime/USDCore/Source/USDUtilities/Public/USDTypesConversion.h
--- a/Engine/Plugins/Runtime/USDCore/Source/USDUtilities/Public/USDTypesConversion.h
+++ b/Engine/Plugins/Runtime/USDCore/Source/USDUtilities/Public/USDTypesConversion.h
@@ -178,6 +178,9 @@ namespace UsdUtils
 	USDUTILITIES_API FTransform ConvertTransformToUsdSpace(const FUsdStageInfo& StageInfo, const FTransform& TransformInUESpace);
 	USDUTILITIES_API FTransform ConvertTransformToUESpace(const FUsdStageInfo& StageInfo, const FTransform& TransformInUsdSpace);
 
+	/** Produces a transform that can be applied to a point from USD space to convert it into UE space */
+	USDUTILITIES_API FTransform GetUsdToUESpaceTransform(const FUsdStageInfo& StageInfo);
+
 	FTransform ConvertAxes(const bool bZUp, const FTransform Transform);
 }
 #endif	  // #if USE_USD_SDK
