USD: Fix ensure when stage uses a path to a texture that doesn't exist.

diff --git a/Engine/Plugins/Interchange/Extensions/OpenUSD/Source/Import/Private/InterchangeUsdTranslator.cpp b/Engine/Plugins/Interchange/Extensions/OpenUSD/Source/Import/Private/InterchangeUsdTranslator.cpp
--- a/Engine/Plugins/Interchange/Extensions/OpenUSD/Source/Import/Private/InterchangeUsdTranslator.cpp
+++ b/Engine/Plugins/Interchange/Extensions/OpenUSD/Source/Import/Private/InterchangeUsdTranslator.cpp
@@ -2405,6 +2405,27 @@ namespace UE::InterchangeUsdTranslator::Private
 				LightNode->SetCustomCubemapDependency(TextureCubeNodeUid);
 				LightNode->SetCustomSourceType(EInterchangeSkyLightSourceType::SpecifiedCubemap);
 			}
+
+			// If this file doesn't exist then retrieving the payload will fail, but we emit the warning from here as
+			// we can provide a better error message than the the payload retrieval function
+			if (!FPaths::FileExists(ResolvedDomeTexturePath))
+			{
+				// Prefer to retrieve the original filepath string as this ResolvedDomeTexturePath will likely just
+				// be the empty string if we failed to resolve
+				FString RawFilePath = ResolvedDomeTexturePath;
+				if (FileAttribute)
+				{
+					FileAttribute.Get(RawFilePath);
+				}
+
+				UE_LOG(
+					LogUsd,
+					Warning,
+					TEXT("Failed to find the file '%s' specified by DomeLight prim '%s'"),
+					*RawFilePath,
+					*Prim.GetPrimPath().GetString()
+				);
+			}
 		}
 
 		return NodeUid;
diff --git a/Engine/Plugins/Interchange/Runtime/Source/Import/Private/MaterialX/InterchangeMaterialXTranslator.cpp b/Engine/Plugins/Interchange/Runtime/Source/Import/Private/MaterialX/InterchangeMaterialXTranslator.cpp
--- a/Engine/Plugins/Interchange/Runtime/Source/Import/Private/MaterialX/InterchangeMaterialXTranslator.cpp
+++ b/Engine/Plugins/Interchange/Runtime/Source/Import/Private/MaterialX/InterchangeMaterialXTranslator.cpp
@@ -80,7 +80,7 @@ TOptional<UE::Interchange::FImportImage> UInterchangeMaterialXTranslator::GetTex
 	UE::Interchange::Private::FScopedTranslator ScopedTranslator(Filename, InResults, InAnalyticsHelper);
 	const IInterchangeTexturePayloadInterface* TextureTranslator = ScopedTranslator.GetPayLoadInterface<IInterchangeTexturePayloadInterface>();
 
-	if (!ensure(TextureTranslator))
+	if (!TextureTranslator)
 	{
 		return TOptional<UE::Interchange::FImportImage>();
 	}
diff --git a/Engine/Plugins/Runtime/USDCore/Source/USDUtilities/Private/USDShadeConversion.cpp b/Engine/Plugins/Runtime/USDCore/Source/USDUtilities/Private/USDShadeConversion.cpp
--- a/Engine/Plugins/Runtime/USDCore/Source/USDUtilities/Private/USDShadeConversion.cpp
+++ b/Engine/Plugins/Runtime/USDCore/Source/USDUtilities/Private/USDShadeConversion.cpp
@@ -616,6 +616,21 @@ namespace UE::UsdShadeConversion::Private
 				}
 				else
 				{
+					// Prefer to retrieve the original filepath string as this TexturePath will likely just
+					// be the empty string if we failed to resolve
+					FString RawFilePath = TexturePath;
+					if (FileInput)
+					{
+						pxr::SdfAssetPath AssetPath;
+						FileInput.Get<pxr::SdfAssetPath>(&AssetPath);
+						RawFilePath = UsdToUnreal::ConvertString(AssetPath.GetAssetPath());
+					}
+
+					USD_LOG_WARNING(
+						TEXT("Failed to find the file '%s', used by shader prim '%s'."),
+						*RawFilePath,
+						*UsdToUnreal::ConvertPath(InputOrOutput.GetPrim().GetPrimPath())
+					);
 					return false;
 				}
 
