[Interchange-USD] Making HISM component option only visible for USD for Scene Level import.

#rnx
#virtualized

diff --git a/Engine/Plugins/Interchange/Extensions/OpenUSD/Source/Import/Private/InterchangeOpenUSDImportModule.cpp b/Engine/Plugins/Interchange/Extensions/OpenUSD/Source/Import/Private/InterchangeOpenUSDImportModule.cpp
--- a/Engine/Plugins/Interchange/Extensions/OpenUSD/Source/Import/Private/InterchangeOpenUSDImportModule.cpp
+++ b/Engine/Plugins/Interchange/Extensions/OpenUSD/Source/Import/Private/InterchangeOpenUSDImportModule.cpp
@@ -56,7 +56,7 @@ void FInterchangeOpenUSDImportModule::StartupModule()
 			FInterchangeTranslatorPipelines TranslatorPipelines;
 			TranslatorPipelines.Translator = TranslatorClassPath;
 			TranslatorPipelines.Pipelines.Add(FSoftObjectPath(TEXT("/InterchangeOpenUSD/Pipelines/DefaultUSDSceneAssetsPipeline.DefaultUSDSceneAssetsPipeline")));
-			TranslatorPipelines.Pipelines.Add(FSoftObjectPath(TEXT("/Interchange/Pipelines/DefaultSceneLevelPipeline.DefaultSceneLevelPipeline")));
+			TranslatorPipelines.Pipelines.Add(FSoftObjectPath(TEXT("/InterchangeOpenUSD/Pipelines/DefaultUSDSceneLevelPipeline.DefaultUSDSceneLevelPipeline")));
 			TranslatorPipelines.Pipelines.Add(FSoftObjectPath(TEXT("/InterchangeOpenUSD/Pipelines/DefaultUSDPipeline.DefaultUSDPipeline")));
 			TranslatorPipelines.Pipelines.Add(FSoftObjectPath(TEXT("/Interchange/Pipelines/DefaultMaterialXPipeline.DefaultMaterialXPipeline")));
 
diff --git a/Engine/Plugins/Interchange/Runtime/Source/Pipelines/Public/InterchangeGenericScenesPipeline.h b/Engine/Plugins/Interchange/Runtime/Source/Pipelines/Public/InterchangeGenericScenesPipeline.h
--- a/Engine/Plugins/Interchange/Runtime/Source/Pipelines/Public/InterchangeGenericScenesPipeline.h
+++ b/Engine/Plugins/Interchange/Runtime/Source/Pipelines/Public/InterchangeGenericScenesPipeline.h
@@ -69,6 +69,7 @@ public:
 	UPROPERTY(EditAnywhere, BlueprintReadWrite, Category = "Scene", meta = (SubCategory = "Reimport Assets"))
 	bool bDeleteMissingAssets = false;
 
+	/* If enabled, HierarchicalInstancedStaticMeshComponents will be generated on import instead of InstancedStaticMeshComponents. */
 	UPROPERTY(EditAnywhere, BlueprintReadWrite, Category = "Scene", meta = (SubCategory = "Component properties", DisplayName = "Use Hierarchical ISM Components"))
 	bool bUseHierarchicalISMComponents = false;
 
