USD: Allow point instancer traversal to step through unknown prims;

Fix crash when point instancer prototype prim cannot be found.

diff --git a/Engine/Plugins/Interchange/Extensions/OpenUSD/Source/Import/Private/InterchangeUsdTranslator.cpp b/Engine/Plugins/Interchange/Extensions/OpenUSD/Source/Import/Private/InterchangeUsdTranslator.cpp
--- a/Engine/Plugins/Interchange/Extensions/OpenUSD/Source/Import/Private/InterchangeUsdTranslator.cpp
+++ b/Engine/Plugins/Interchange/Extensions/OpenUSD/Source/Import/Private/InterchangeUsdTranslator.cpp
@@ -164,6 +164,11 @@ namespace UE::InterchangeUsdTranslator::Private
 	public:
 		void UpdateWithCurrentPrim(const UE::FUsdPrim& CurrentPrim)
 		{
+			if (!CurrentPrim)
+			{
+				return;
+			}
+
 			bVisible = bVisible && UsdUtils::HasInheritedVisibility(CurrentPrim);
 
 			// Check this first so that we go bInsideLOD if our parent was a LODContainer
@@ -2363,14 +2368,6 @@ namespace UE::InterchangeUsdTranslator::Private
 					}
 				}
 			}
-			else if (Prim.IsA(TEXT("Xform")))
-			{
-				TArray<FUsdPrim> Children = Prim.GetChildren();
-				for (const FUsdPrim& Child : Children)
-				{
-					TraverseNestedPointInstancerHierarchy(Child, Info, CurrentGlobalTransform);
-				}
-			}
 			else if (Prim.IsA(TEXT("Mesh")))
 			{
 				if (Info.bVisible)
@@ -2388,6 +2385,14 @@ namespace UE::InterchangeUsdTranslator::Private
 					ISMComponentNode->SetCustomInstancedAssetUid(MeshUid);
 				}
 			}
+			else
+			{
+				TArray<FUsdPrim> Children = Prim.GetChildren();
+				for (const FUsdPrim& Child : Children)
+				{
+					TraverseNestedPointInstancerHierarchy(Child, Info, CurrentGlobalTransform);
+				}
+			}
 		}
 
 		void TraversePointInstancerHierarchy(
@@ -2408,21 +2413,16 @@ namespace UE::InterchangeUsdTranslator::Private
 				FString PrototypePathString = PrototypeInstanceData.PathString;
 
 				UE::FUsdPrim PrototypePrim = Prim.GetStage().GetPrimAtPath(UE::FSdfPath{ *PrototypePathString });
+				if (!PrototypePrim)
+				{
+					UE_LOG(LogUsd, Warning, TEXT("Failed to find prototype '%s' for point instancer prim '%s'"), *PrototypePathString, *Prim.GetPrimPath().GetString());
+					continue;
+				}
 
 				FTraversalInfo PrototypePrimInfo = Info;
 				PrototypePrimInfo.UpdateWithCurrentPrim(PrototypePrim);
 
-				if (PrototypePrim.IsA(TEXT("PointInstancer")) || PrototypePrim.IsA(TEXT("Xform")))
-				{
-					for (int32 InstanceCounter = 0; InstanceCounter < PrototypeInstanceTransforms.Num(); ++InstanceCounter)
-					{
-						//Note: For handling the point instancer's instance's transform we create a component here. This is an extra introduced component compared to the original hierarchy.
-						FTransform CurrentGlobalTransform = PrototypeInstanceTransforms[InstanceCounter] * ParentGlobalTransform;
-
-						TraverseNestedPointInstancerHierarchy(PrototypePrim, PrototypePrimInfo, CurrentGlobalTransform);
-					}
-				}
-				else if (PrototypePrim.IsA(TEXT("Mesh")))
+				if (PrototypePrim.IsA(TEXT("Mesh")))
 				{
 					if (PrototypePrimInfo.bVisible)
 					{
@@ -2447,6 +2447,16 @@ namespace UE::InterchangeUsdTranslator::Private
 						ISMComponentNode->SetCustomInstancedAssetUid(MeshUid);
 					}
 				}
+				else
+				{
+					for (int32 InstanceCounter = 0; InstanceCounter < PrototypeInstanceTransforms.Num(); ++InstanceCounter)
+					{
+						//Note: For handling the point instancer's instance's transform we create a component here. This is an extra introduced component compared to the original hierarchy.
+						FTransform CurrentGlobalTransform = PrototypeInstanceTransforms[InstanceCounter] * ParentGlobalTransform;
+
+						TraverseNestedPointInstancerHierarchy(PrototypePrim, PrototypePrimInfo, CurrentGlobalTransform);
+					}
+				}
 			}
 		}
 	};
