Interchange USD: Fixed Interchange USD translator was holding on to cached stage causing it to not properly reload external references.

#rnx

diff --git a/Engine/Plugins/Interchange/Extensions/OpenUSD/Source/Import/Private/InterchangeUsdContext.cpp b/Engine/Plugins/Interchange/Extensions/OpenUSD/Source/Import/Private/InterchangeUsdContext.cpp
--- a/Engine/Plugins/Interchange/Extensions/OpenUSD/Source/Import/Private/InterchangeUsdContext.cpp
+++ b/Engine/Plugins/Interchange/Extensions/OpenUSD/Source/Import/Private/InterchangeUsdContext.cpp
@@ -14,16 +14,22 @@ UInterchangeUsdContext::UInterchangeUsdContext()
 {
 }
 
-void UInterchangeUsdContext::BeginDestroy()
+void UInterchangeUsdContext::Reset()
 {
 #if USE_USD_SDK
 	if (bShouldCleanUpFromStageCache && StageIdInUsdUtilsStageCache != INDEX_NONE)
 	{
 		UsdUtils::RemoveStageFromUsdUtilsStageCache(StageIdInUsdUtilsStageCache);
+		StageIdInUsdUtilsStageCache = INDEX_NONE;
 	}
 
 	ReleaseInfoCache();
 #endif	  // USE_USD_SDK
+}
+
+void UInterchangeUsdContext::BeginDestroy()
+{
+	Reset();
 
 	Super::BeginDestroy();
 }
diff --git a/Engine/Plugins/Interchange/Extensions/OpenUSD/Source/Import/Private/InterchangeUsdTranslator.cpp b/Engine/Plugins/Interchange/Extensions/OpenUSD/Source/Import/Private/InterchangeUsdTranslator.cpp
--- a/Engine/Plugins/Interchange/Extensions/OpenUSD/Source/Import/Private/InterchangeUsdTranslator.cpp
+++ b/Engine/Plugins/Interchange/Extensions/OpenUSD/Source/Import/Private/InterchangeUsdTranslator.cpp
@@ -6517,40 +6517,39 @@ bool UInterchangeUSDTranslator::Translate(UInterchangeBaseNodeContainer& NodeCon
 #endif	  // USE_USD_SDK
 }
 
-void UInterchangeUSDTranslator::ReleaseSource()
+void UInterchangeUSDTranslator::ImportFinish()
 {
-#if USE_USD_SDK
-	UE::InterchangeUsdTranslator::Private::UInterchangeUSDTranslatorImpl* ImplPtr = Impl.Get();
-	if (ImplPtr)
-	{
-		ImplPtr->TranslationContext.Reset();	// This holds a reference to the stage
-		ImplPtr->UsdStage = UE::FUsdStage{};
-		ImplPtr->InfoCache = nullptr;
-		ImplPtr->CurrentTrackSet = nullptr;
-
-		{
-			FWriteScopeLock Lock{ImplPtr->PrimPathToVariantToStageLock};
-			ImplPtr->PrimPathToVariantToStage.Reset();
-		}
-	}
-#endif	  // USE_USD_SDK
-
+	// These objects are only set by Translate so they are released here instead of ReleaseSource
 	if (TranslatorSettings)
 	{
 		TranslatorSettings->ClearFlags(RF_Standalone);
 		TranslatorSettings = nullptr;
 	}
-}
 
-void UInterchangeUSDTranslator::ImportFinish()
-{
 #if USE_USD_SDK
+	// The USD context is caching the stage so it must be released once the import is done
+	UInterchangeUsdContext* ContextObject = Cast<UInterchangeUsdContext>(SourceData->GetContextObjectByTag(UE::Interchange::USD::USDContextTag));
+	if (ContextObject)
+	{
+		ContextObject->Reset();
+	}
+
 	UE::InterchangeUsdTranslator::Private::UInterchangeUSDTranslatorImpl* ImplPtr = Impl.Get();
 	if (!ImplPtr)
 	{
 		return;
 	}
 
+	ImplPtr->TranslationContext.Reset();	// This holds a reference to the stage
+	ImplPtr->UsdStage = UE::FUsdStage{};
+	ImplPtr->InfoCache = nullptr;
+	ImplPtr->CurrentTrackSet = nullptr;
+
+	{
+		FWriteScopeLock Lock{ ImplPtr->PrimPathToVariantToStageLock };
+		ImplPtr->PrimPathToVariantToStage.Reset();
+	}
+
 	ImplPtr->CleanUpDecompressedUSDZFolder();
 #endif	  // USE_USD_SDK
 }
diff --git a/Engine/Plugins/Interchange/Extensions/OpenUSD/Source/Import/Public/InterchangeUsdContext.h b/Engine/Plugins/Interchange/Extensions/OpenUSD/Source/Import/Public/InterchangeUsdContext.h
--- a/Engine/Plugins/Interchange/Extensions/OpenUSD/Source/Import/Public/InterchangeUsdContext.h
+++ b/Engine/Plugins/Interchange/Extensions/OpenUSD/Source/Import/Public/InterchangeUsdContext.h
@@ -46,6 +46,8 @@ public:
 	// Properly releases the reference to the current stage cache, whether owned or not (deleting it if owned)
 	void ReleaseInfoCache();
 
+	void Reset();
+
 private:
 	// We never store the stage itself, but only it's Id within the UsdUtils' singleton StageCache.
 	// The intent here is to allow Python stages to be passed in and manipulated via Python in case of Python imports or pipelines.
diff --git a/Engine/Plugins/Interchange/Extensions/OpenUSD/Source/Import/Public/InterchangeUsdTranslator.h b/Engine/Plugins/Interchange/Extensions/OpenUSD/Source/Import/Public/InterchangeUsdTranslator.h
--- a/Engine/Plugins/Interchange/Extensions/OpenUSD/Source/Import/Public/InterchangeUsdTranslator.h
+++ b/Engine/Plugins/Interchange/Extensions/OpenUSD/Source/Import/Public/InterchangeUsdTranslator.h
@@ -108,7 +108,6 @@ public:
 	virtual EInterchangeTranslatorAssetType GetSupportedAssetTypes() const override;
 	virtual TArray<FString> GetSupportedFormats() const override;
 	virtual bool Translate(UInterchangeBaseNodeContainer& BaseNodeContainer) const override;
-	virtual void ReleaseSource() override;
 	virtual void ImportFinish() override;
 	virtual UInterchangeTranslatorSettings* GetSettings() const override;
 	virtual void SetSettings(const UInterchangeTranslatorSettings* InterchangeTranslatorSettings) override;
diff --git a/Engine/Plugins/Runtime/USDCore/Source/UnrealUSDWrapper/Private/UnrealUSDWrapper.cpp b/Engine/Plugins/Runtime/USDCore/Source/UnrealUSDWrapper/Private/UnrealUSDWrapper.cpp
--- a/Engine/Plugins/Runtime/USDCore/Source/UnrealUSDWrapper/Private/UnrealUSDWrapper.cpp
+++ b/Engine/Plugins/Runtime/USDCore/Source/UnrealUSDWrapper/Private/UnrealUSDWrapper.cpp
@@ -945,7 +945,8 @@ namespace UE::UnrealUSDWrapper::Private
 			// Layers are cached in the layer registry independently of the stage cache. If the layer is already in
 			// the registry by the time we try to open a stage, even if we're not using a stage cache at all the
 			// layer will be reused and the file will *not* be re-read, so here we manually reload them.
-			pxr::SdfLayerHandleVector StageLayers = Stage->GetLayerStack();
+			// GetUsedLayers also returns layers from external references and such
+			pxr::SdfLayerHandleVector StageLayers = Stage->GetUsedLayers();
 			for (pxr::SdfLayerHandle StageLayer : StageLayers)
 			{
 				if (LoadedLayers.count(StageLayer) > 0)
