Interchange USD: Add a case sensitive hash to generated node UIDs.

diff --git a/Engine/Plugins/Interchange/Extensions/OpenUSD/Source/Import/Private/InterchangeUSDPipeline.cpp b/Engine/Plugins/Interchange/Extensions/OpenUSD/Source/Import/Private/InterchangeUSDPipeline.cpp
--- a/Engine/Plugins/Interchange/Extensions/OpenUSD/Source/Import/Private/InterchangeUSDPipeline.cpp
+++ b/Engine/Plugins/Interchange/Extensions/OpenUSD/Source/Import/Private/InterchangeUSDPipeline.cpp
@@ -50,7 +50,7 @@ namespace UE::InterchangeUsdPipeline::Private
 #if USE_USD_SDK
 	const FString& GetPseudoRootTranslatedNodeUid()
 	{
-		const static FString Result = UE::FSdfPath::AbsoluteRootPath().GetString();
+		const static FString Result = UE::Interchange::USD::MakeNodeUid(UE::FSdfPath::AbsoluteRootPath().GetString());
 		return Result;
 	}
 
diff --git a/Engine/Plugins/Interchange/Extensions/OpenUSD/Source/Import/Private/InterchangeUsdTranslator.cpp b/Engine/Plugins/Interchange/Extensions/OpenUSD/Source/Import/Private/InterchangeUsdTranslator.cpp
--- a/Engine/Plugins/Interchange/Extensions/OpenUSD/Source/Import/Private/InterchangeUsdTranslator.cpp
+++ b/Engine/Plugins/Interchange/Extensions/OpenUSD/Source/Import/Private/InterchangeUsdTranslator.cpp
@@ -117,6 +117,8 @@ static FAutoConsoleVariableRef CVarInterchangeEnableUSDAudioImport(
 
 namespace UE::InterchangeUsdTranslator::Private
 {
+	using namespace UE::Interchange::USD;
+
 #if USE_USD_SDK
 	const static FString AnimationPrefix = TEXT("\\Animation\\");
 	const static FString AnimationTrackPrefix = TEXT("\\AnimationTrack\\");
@@ -626,7 +628,7 @@ namespace UE::InterchangeUsdTranslator::Private
 
 	FString GetMorphTargetMeshNodeUid(const FString& MeshPrimPath, int32 MeshBlendShapeIndex, const FString& InbetweenName = FString{})
 	{
-		return FString::Printf(TEXT("%s%s\\%d\\%s"), *MorphTargetPrefix, *MeshPrimPath, MeshBlendShapeIndex, *InbetweenName);
+		return MakeNodeUid(FString::Printf(TEXT("%s%s\\%d\\%s"), *MorphTargetPrefix, *MeshPrimPath, MeshBlendShapeIndex, *InbetweenName));
 	}
 
 	FString GetMorphTargetMeshPayloadKey(
@@ -1016,7 +1018,7 @@ namespace UE::InterchangeUsdTranslator::Private
 
 	FString AddTextureCubeNode(const FString& TextureFilePath, UInterchangeBaseNodeContainer& NodeContainer)
 	{
-		const FString NodeUid = FString::Printf(TEXT("%s%s"), *TexturePrefix, *TextureFilePath);
+		const FString NodeUid = MakeNodeUid(FString::Printf(TEXT("%s%s"), *TexturePrefix, *TextureFilePath));
 
 		// If we created an actual TextureCubeNode here we'd be queried for the sliced payload later. Our HDR image is
 		// just a latlong 2d texture though, so we instead need to create a Texture2DNode and then set SetForceLongLatCubemap(),
@@ -1047,7 +1049,7 @@ namespace UE::InterchangeUsdTranslator::Private
 		using namespace UsdUnreal::MaterialUtils;
 
 		FString DisplayColorDesc = Slot.MaterialSource;
-		FString NodeUid = MaterialPrefix + Slot.MaterialSource;
+		FString NodeUid = MakeNodeUid(MaterialPrefix + Slot.MaterialSource);
 
 		// We'll treat the DisplayColorDesc (something like "!DisplayColor_1_0") as the material instance UID here
 		const UInterchangeMaterialInstanceNode* Node = Cast<UInterchangeMaterialInstanceNode>(NodeContainer.GetNode(NodeUid));
@@ -1080,7 +1082,7 @@ namespace UE::InterchangeUsdTranslator::Private
 		using namespace UsdUnreal::MaterialUtils;
 
 		// e.g. "\\MaterialReference\\/Game/MyFolder/Red.Red"
-		const FString NodeUid = MaterialReferencePrefix + ContentPath;
+		const FString NodeUid = MakeNodeUid(MaterialReferencePrefix + ContentPath);
 
 		const UInterchangeMaterialReferenceNode* Node = Cast<UInterchangeMaterialReferenceNode>(NodeContainer.GetNode(NodeUid));
 		if (Node)
@@ -2195,7 +2197,7 @@ namespace UE::InterchangeUsdTranslator::Private
 
 		UInterchangeInstancedStaticMeshComponentNode* CreateSceneComponent(const FString& MeshUid)
 		{
-			FString UniqueId = ComponentPrefix + TEXT("\\") + ParentSceneNodeUid + TEXT("\\") + MeshUid;
+			FString UniqueId = MakeNodeUid(ComponentPrefix + TEXT("\\") + ParentSceneNodeUid + TEXT("\\") + MeshUid);
 
 			UInterchangeInstancedStaticMeshComponentNode* ISMComponentNode = UidToComponentNode.FindRef(UniqueId);
 
@@ -2937,7 +2939,7 @@ namespace UE::InterchangeUsdTranslator::Private
 
 		FString NodeUid;
 
-		NodeUid = LightPrefix + Prim.GetPrimPath().GetString();
+		NodeUid = MakeNodeUid(LightPrefix + Prim.GetPrimPath().GetString());
 		FString NodeName(Prim.GetName().ToString());
 
 		// Ref. UsdToUnreal::ConvertLight
@@ -3166,7 +3168,7 @@ namespace UE::InterchangeUsdTranslator::Private
 
 		FString NodeUid;
 
-		NodeUid = CameraPrefix + Prim.GetPrimPath().GetString();
+		NodeUid = CameraPrefix + MakeNodeUid(Prim.GetPrimPath().GetString());
 		FString NodeName(Prim.GetName().ToString());
 
 		UInterchangePhysicalCameraNode* CameraNode = NewObject<UInterchangePhysicalCameraNode>(&NodeContainer);
@@ -3330,7 +3332,7 @@ namespace UE::InterchangeUsdTranslator::Private
 
 		const UE::FUsdStage Stage = PrototypePrim.GetStage();
 		const FString PrototypePrimPath = PrototypePrim.GetPrimPath().GetString();
-		const FString NodeUid = (bPrimitiveShape ? PrimitiveShapePrefix : MeshPrefix) + PrototypePrimPath;
+		const FString NodeUid = MakeNodeUid((bPrimitiveShape ? PrimitiveShapePrefix : MeshPrefix) + PrototypePrimPath);
 		const FString NodeName(PrototypePrim.GetName().ToString());
 
 		// Check if Node already exist with this ID
@@ -3404,7 +3406,7 @@ namespace UE::InterchangeUsdTranslator::Private
 			MeshNode->SetPayLoadKey(PayloadKey, EInterchangeMeshPayLoadType::SKELETAL);
 			if (Info.BoundSkeletonPrimPath && !Info.BoundSkeletonPrimPath->IsEmpty())
 			{
-				MeshNode->SetSkeletonDependencyUid(*Info.BoundSkeletonPrimPath);
+				MeshNode->SetSkeletonDependencyUid(MakeNodeUid(*Info.BoundSkeletonPrimPath));
 			}
 
 			AddMorphTargetNodes(PrototypePrim, *this, *MeshNode, NodeContainer, Info);
@@ -3741,7 +3743,7 @@ namespace UE::InterchangeUsdTranslator::Private
 		}
 
 		FString MaterialDisplayLabel = bIsFallback ? TEXT("USDVolumetricFallbackMaterial") : FPaths::GetBaseFilename(ParentMaterialPath);
-		FString MaterialNodeUid = MaterialPrefix + MaterialDisplayLabel;
+		FString MaterialNodeUid = MakeNodeUid(MaterialPrefix + MaterialDisplayLabel);
 		MakeNodeUidUniqueInContainer(MaterialNodeUid, NodeContainer);
 
 		// We'll always spawn a new material instance for each volume prim. Realistically a stage is only going to have a handful
@@ -4098,7 +4100,7 @@ namespace UE::InterchangeUsdTranslator::Private
 		TRACE_CPUPROFILER_EVENT_SCOPE(UInterchangeUSDTranslatorImpl::AddGroomNode)
 
 		const FString PrimPath = UsdUtils::GetPrototypePrimPath(Prim).GetString();
-		const FString NodeUid = GroomPrefix + PrimPath;
+		const FString NodeUid = MakeNodeUid(GroomPrefix + PrimPath);
 		const FString NodeName(Prim.GetName().ToString());
 
 		// Check if Node already exist with this ID
@@ -4151,7 +4153,7 @@ namespace UE::InterchangeUsdTranslator::Private
 		TRACE_CPUPROFILER_EVENT_SCOPE(UInterchangeUSDTranslatorImpl::AddSpatialAudioNode)
 
 		const FString PrimPath = UsdUtils::GetPrototypePrimPath(Prim).GetString();
-		const FString NodeUid = SpatialAudioPrefix + PrimPath;
+		const FString NodeUid = MakeNodeUid(SpatialAudioPrefix + PrimPath);
 		const FString NodeName(Prim.GetName().ToString());
 
 		// Check if Node already exist with this ID
@@ -4186,7 +4188,7 @@ namespace UE::InterchangeUsdTranslator::Private
 		}
 
 		UE::FSdfLayer Layer = Impl.UsdStage.GetRootLayer();
-		const FString AnimTrackSetNodeUid = AnimationPrefix + Layer.GetIdentifier();
+		const FString AnimTrackSetNodeUid = MakeNodeUid(AnimationPrefix + Layer.GetIdentifier());
 		const FString AnimTrackSetNodeDisplayName = FPaths::GetBaseFilename(Layer.GetDisplayName());	// Strip extension
 
 		// We should only have one track set node per scene for now
@@ -4211,7 +4213,8 @@ namespace UE::InterchangeUsdTranslator::Private
 
 		const FString PrimPath = Prim.GetPrimPath().GetString();
 		const FString UniquePath = PrimPath + TEXT("\\") + UnrealIdentifiers::TransformPropertyName.ToString();
-		const FString AnimTrackNodeUid = AnimationTrackPrefix + UniquePath;
+		const FString AnimTrackNodeUid = MakeNodeUid(AnimationTrackPrefix + UniquePath);
+		const FString SceneNodeUid = MakeNodeUid(PrimPath);
 
 		const UInterchangeTransformAnimationTrackNode* ExistingNode = Cast<UInterchangeTransformAnimationTrackNode>(
 			NodeContainer.GetNode(AnimTrackNodeUid)
@@ -4223,7 +4226,7 @@ namespace UE::InterchangeUsdTranslator::Private
 
 		UInterchangeTransformAnimationTrackNode* TransformAnimTrackNode = NewObject<UInterchangeTransformAnimationTrackNode>(&NodeContainer);
 		NodeContainer.SetupNode(TransformAnimTrackNode, AnimTrackNodeUid, UniquePath, EInterchangeNodeContainerType::TranslatedAsset);
-		TransformAnimTrackNode->SetCustomActorDependencyUid(*PrimPath);
+		TransformAnimTrackNode->SetCustomActorDependencyUid(SceneNodeUid);
 		TransformAnimTrackNode->SetCustomAnimationPayloadKey(UniquePath, EInterchangeAnimationPayLoadType::CURVE);
 		TransformAnimTrackNode->SetCustomUsedChannels((int32)EMovieSceneTransformChannel::AllTransform);
 
@@ -4232,7 +4235,7 @@ namespace UE::InterchangeUsdTranslator::Private
 	}
 
 	void AddPropertyAnimationNode(
-		const FString& SceneNodeUid,
+		const FString& PrimPath,
 		const FString& UEPropertyName,
 		EInterchangePropertyTracks TrackType,
 		EInterchangeAnimationPayLoadType PayloadType,
@@ -4243,8 +4246,9 @@ namespace UE::InterchangeUsdTranslator::Private
 		// We don't use the USD attribute path here because we want one unique node per UE track name,
 		// so that if e.g. both "intensity" and "exposure" are animated we make a single track for
 		// the Intensity UE property
-		const FString UniquePath = SceneNodeUid + TEXT("\\") + UEPropertyName;
-		const FString AnimTrackNodeUid = AnimationTrackPrefix + UniquePath;
+		const FString PayloadKey = PrimPath + TEXT("\\") + UEPropertyName;
+		const FString SceneNodeUid = MakeNodeUid(PrimPath);
+		const FString AnimTrackNodeUid = MakeNodeUid(AnimationTrackPrefix + PayloadKey);
 
 		const UInterchangeAnimationTrackNode* ExistingNode = Cast<UInterchangeAnimationTrackNode>(NodeContainer.GetNode(AnimTrackNodeUid));
 		if (ExistingNode)
@@ -4253,10 +4257,10 @@ namespace UE::InterchangeUsdTranslator::Private
 		}
 
 		UInterchangeAnimationTrackNode* AnimTrackNode = NewObject<UInterchangeAnimationTrackNode>(&NodeContainer);
-		NodeContainer.SetupNode(AnimTrackNode, AnimTrackNodeUid, UniquePath, EInterchangeNodeContainerType::TranslatedAsset);
+		NodeContainer.SetupNode(AnimTrackNode, AnimTrackNodeUid, PayloadKey, EInterchangeNodeContainerType::TranslatedAsset);
 		AnimTrackNode->SetCustomActorDependencyUid(*SceneNodeUid);
 		AnimTrackNode->SetCustomPropertyTrack(TrackType);
-		AnimTrackNode->SetCustomAnimationPayloadKey(UniquePath, PayloadType);
+		AnimTrackNode->SetCustomAnimationPayloadKey(PayloadKey, PayloadType);
 
 		AddTrackSetNode(Impl, NodeContainer);
 		Impl.CurrentTrackSet->AddCustomAnimationTrackUid(AnimTrackNodeUid);
@@ -4323,7 +4327,7 @@ namespace UE::InterchangeUsdTranslator::Private
 		{
 			return;
 		}
-		const FString& SceneNodeUid = SceneNode->GetUniqueID();
+		const FString& PrimPath = Prim.GetPrimPath().GetString();
 
 		// Target the scene node from all the volume nodes to make it easy to find everything we need on the USD Pipeline
 		for (const FString& AssetNodeUid : AssetNodeUids)
@@ -4347,7 +4351,7 @@ namespace UE::InterchangeUsdTranslator::Private
 			const EInterchangePropertyTracks TrackType = EInterchangePropertyTracks::HeterogeneousVolumeFrame;
 			const static FName UEPropertyName = GET_MEMBER_NAME_CHECKED(UHeterogeneousVolumeComponent, Frame);
 
-			AddPropertyAnimationNode(SceneNodeUid, UEPropertyName.ToString(), TrackType, PayloadType, Impl, NodeContainer);
+			AddPropertyAnimationNode(PrimPath, UEPropertyName.ToString(), TrackType, PayloadType, Impl, NodeContainer);
 		}
 	}
 
@@ -4384,7 +4388,7 @@ namespace UE::InterchangeUsdTranslator::Private
 		const FString SkelAnimationPrimPath = UsdUtils::GetPrototypePrimPath(SkelAnimationPrim).GetString();
 		const FString SkeletonPrimPath = SkeletonPrim.GetPrimPath().GetString();
 		const FString UniquePath = SkelAnimationPrimPath + TEXT("\\") + SkeletonPrimPath;
-		const FString NodeUid = AnimationTrackPrefix + UniquePath;
+		const FString NodeUid = MakeNodeUid(AnimationTrackPrefix + UniquePath);
 
 		const UInterchangeSkeletalAnimationTrackNode* ExistingNode = Cast<UInterchangeSkeletalAnimationTrackNode>(NodeContainer.GetNode(NodeUid));
 		if (ExistingNode)
@@ -4628,7 +4632,7 @@ namespace UE::InterchangeUsdTranslator::Private
 
 			// Putting the BonePrefix here avoids the pathological case where the user has skeleton child prims
 			// with names that match the joint names
-			const FString BoneNodeUid = SkeletonPrimNodeUid + BonePrefix + ConcatBonePath;
+			const FString BoneNodeUid = MakeNodeUid(SkeletonPrimNodeUid + BonePrefix + ConcatBonePath);
 
 			UInterchangeSceneNode* BoneNode = NewObject<UInterchangeSceneNode>(&NodeContainer);
 			NodeContainer.SetupNode(BoneNode, BoneNodeUid, Bone.Name, EInterchangeNodeContainerType::TranslatedScene, ParentNode.GetUniqueID());
@@ -5049,7 +5053,7 @@ namespace UE::InterchangeUsdTranslator::Private
 	{
 		TRACE_CPUPROFILER_EVENT_SCOPE(Traverse)
 
-		const FString SceneNodeUid = Prim.GetPrimPath().GetString();
+		const FString SceneNodeUid = MakeNodeUid(Prim.GetPrimPath().GetString());
 		const FString DisplayLabel(Prim.GetName().ToString());
 		const FName TypeName = Prim.GetTypeName();
 
@@ -5649,7 +5653,7 @@ namespace UE::InterchangeUsdTranslator::Private
 		// Interchange can instance skeletal meshes we should use the line below to obtain the node uid
 		//  (in the same way static mesh does already).
 		// const FString& MeshNodeUid = MeshPrefix + UsdUtils::GetPrototypePrimPath(Prim).GetString();
-		const FString& MeshNodeUid = MeshPrefix + Prim.GetPrimPath().GetString();
+		const FString& MeshNodeUid = MakeNodeUid(MeshPrefix + Prim.GetPrimPath().GetString());
 
 		// Read these variables from the data we cached during traversal for translation
 		TSharedPtr<TArray<FString>> JointNames = nullptr;
@@ -5671,7 +5675,7 @@ namespace UE::InterchangeUsdTranslator::Private
 
 			// The above fields are associated to the mesh *asset* node Uid (hence the prefix),
 			// while the joint names are associated to the skeleton *scene* node Uid, so no prefix
-			const FString& SkeletonNodeUid = *MeshInfo->BoundSkeletonPrimPath;
+			const FString& SkeletonNodeUid = MakeNodeUid(*MeshInfo->BoundSkeletonPrimPath);
 			const FTraversalInfo* SkeletonInfo = Impl.NodeUidToCachedTraversalInfo.Find(SkeletonNodeUid);
 			if (!SkeletonInfo)
 			{
@@ -5847,7 +5851,7 @@ namespace UE::InterchangeUsdTranslator::Private
 			{
 				FReadScopeLock ReadLock{Impl.CachedTraversalInfoLock};
 
-				const FString& MeshNodeUid = MeshPrefix + UsdUtils::GetPrototypePrimPath(MeshPrim).GetString();
+				const FString& MeshNodeUid = MakeNodeUid(MeshPrefix + UsdUtils::GetPrototypePrimPath(MeshPrim).GetString());
 				const FTraversalInfo* MeshInfo = Impl.NodeUidToCachedTraversalInfo.Find(MeshNodeUid);
 				if (MeshInfo)
 				{
@@ -6041,7 +6045,7 @@ namespace UE::InterchangeUsdTranslator::Private
 		{
 			FReadScopeLock ReadLock{Impl.CachedTraversalInfoLock};
 
-			const FTraversalInfo* Info = Impl.NodeUidToCachedTraversalInfo.Find(SkeletonPrimPath);
+			const FTraversalInfo* Info = Impl.NodeUidToCachedTraversalInfo.Find(MakeNodeUid(SkeletonPrimPath));
 			if (!Info)
 			{
 				return false;
@@ -6177,7 +6181,7 @@ namespace UE::InterchangeUsdTranslator::Private
 			{
 				FReadScopeLock ReadLock{Impl.CachedTraversalInfoLock};
 
-				const FTraversalInfo* Info = Impl.NodeUidToCachedTraversalInfo.Find(SkeletonPrimPath);
+				const FTraversalInfo* Info = Impl.NodeUidToCachedTraversalInfo.Find(MakeNodeUid(SkeletonPrimPath));
 				if (!Info)
 				{
 					return false;
diff --git a/Engine/Plugins/Interchange/Runtime/Source/Common/Public/Usd/InterchangeUsdDefinitions.h b/Engine/Plugins/Interchange/Runtime/Source/Common/Public/Usd/InterchangeUsdDefinitions.h
--- a/Engine/Plugins/Interchange/Runtime/Source/Common/Public/Usd/InterchangeUsdDefinitions.h
+++ b/Engine/Plugins/Interchange/Runtime/Source/Common/Public/Usd/InterchangeUsdDefinitions.h
@@ -83,4 +83,15 @@ namespace UE::Interchange::USD
 		// Attribute that informs how we should handle primvars on MeshDescriptions
 		const FString Import = TEXT("USD_Import_Primvars");
 	}
+
+	// Append a case sensistive hash of NodeUid to the NodeUid itself, making it so that any TMap<FString, T> that stores
+	// these IDs behaves in a case sensitive manner.
+	//
+	// This is needed because unfortunately the default TMap<FString, T> is not case sensitive on the key FStrings, while prim
+	// names are case sensitive. Even if we could modify the UInterchangeBaseNodeContainer::Nodes map to be case sensitive, we'd
+	// still constantly get issues as any other TMap<FString, T> in the codebase that stores NodeUids would show collisions
+	inline FString MakeNodeUid(const FString& NodeUid)
+	{
+		return NodeUid + TEXT("_") + LexToString(FCrc::StrCrc32(*NodeUid));
+	}
 }	 // namespace UE::Interchange::USD
