[USD Interchange] Wrap Nanite assembly node generated  UIDs with the new MakeNodeUid()

diff --git a/Engine/Plugins/Interchange/Extensions/OpenUSD/Source/Import/Private/InterchangeUsdTranslator.cpp b/Engine/Plugins/Interchange/Extensions/OpenUSD/Source/Import/Private/InterchangeUsdTranslator.cpp
--- a/Engine/Plugins/Interchange/Extensions/OpenUSD/Source/Import/Private/InterchangeUsdTranslator.cpp
+++ b/Engine/Plugins/Interchange/Extensions/OpenUSD/Source/Import/Private/InterchangeUsdTranslator.cpp
@@ -1600,7 +1600,7 @@ namespace UE::InterchangeUsdTranslator::Private
 			
 			FString MeshUid = TranslatorImpl.AddMeshNode(Prim, TranslatorSettings, NodeContainer, *Info, false, AssemblyDependencies);
 
-			const FString SceneNodeUid = Prim.GetPrimPath().GetString();
+			const FString SceneNodeUid = MakeNodeUid(Prim.GetPrimPath().GetString());
 			const FString DisplayLabel(Prim.GetName().ToString());
 
 			UInterchangeSceneNode* SceneNode = NewObject<UInterchangeSceneNode>(&NodeContainer);
@@ -1789,7 +1789,7 @@ namespace UE::InterchangeUsdTranslator::Private
 		{
 			const bool bIsBaseSkeleton = Prim.GetPrimPath() == BaseSkeletonPrimPath;
 
-			const FString SceneNodeUid = Prim.GetPrimPath().GetString();
+			const FString SceneNodeUid = MakeNodeUid(Prim.GetPrimPath().GetString());
 			const FString DisplayLabel(Prim.GetName().ToString());
 
 			UInterchangeSceneNode* SceneNode = NewObject<UInterchangeSceneNode>(&NodeContainer);
@@ -2134,13 +2134,13 @@ namespace UE::InterchangeUsdTranslator::Private
 			// empty mesh that the mesh factory can eventually add the assembly data to. 
 			else if (MeshTypeIs(ENaniteAssemblyMeshType::StaticMesh))
 			{
-				const FString MeshUid = MeshPrefix + AssemblyRootPrim.GetPrimPath().GetString();
+				const FString MeshUid = MakeNodeUid(MeshPrefix + AssemblyRootPrim.GetPrimPath().GetString());
 				const FString NodeName(AssemblyRootPrim.GetPrimPath().GetName());
 
 				UInterchangeMeshNode* MeshNode = NewObject<UInterchangeMeshNode>(&NodeContainer);
 				NodeContainer.SetupNode(MeshNode, MeshUid, NodeName, EInterchangeNodeContainerType::TranslatedAsset);
 				MeshNode->SetAssetName(NodeName);
-				MeshNode->SetPayLoadKey(ParentSceneNodeUid, EInterchangeMeshPayLoadType::STATIC);
+				MeshNode->SetPayLoadKey(AssemblyRootPrim.GetPrimPath().GetString(), EInterchangeMeshPayLoadType::STATIC);
 
 				for (const FString& PartUid : Result->GetMeshUidsForCategory(EMeshCategory::Part))
 				{
@@ -5535,7 +5535,7 @@ namespace UE::InterchangeUsdTranslator::Private
 
 		// Get Nanite assembly static mesh data, if available
 		
-		const FString MeshNodeUid = MeshPrefix + Prim.GetPrimPath().GetString();
+		const FString MeshNodeUid = MakeNodeUid(MeshPrefix + Prim.GetPrimPath().GetString());
 		FTraversalInfo* MeshInfo;
 		{
 			FReadScopeLock ReadLock{ Impl.CachedTraversalInfoLock };
