Fixed a couple of bughawk issues from CL 46326156.

#rnx

diff --git a/Engine/Plugins/Interchange/Extensions/OpenUSD/Source/Import/Private/InterchangeUsdTranslator.cpp b/Engine/Plugins/Interchange/Extensions/OpenUSD/Source/Import/Private/InterchangeUsdTranslator.cpp
--- a/Engine/Plugins/Interchange/Extensions/OpenUSD/Source/Import/Private/InterchangeUsdTranslator.cpp
+++ b/Engine/Plugins/Interchange/Extensions/OpenUSD/Source/Import/Private/InterchangeUsdTranslator.cpp
@@ -6071,7 +6071,10 @@ namespace UE::InterchangeUsdTranslator::Private
 		}
 
 		// Use the first timecode as an offset when indexing
-		OutPayloadData.RangeStartTime = TimeSampleUnion[0];
+		if (!TimeSampleUnion.IsEmpty())
+		{
+			OutPayloadData.RangeStartTime = TimeSampleUnion[0];
+		}
 
 		const bool bIgnorePrimLocalTransform = false;
 		UsdToUnreal::FPropertyTrackReader Reader = UsdToUnreal::CreatePropertyTrackReader(Prim, UEPropertyName, bIgnorePrimLocalTransform);
diff --git a/Engine/Plugins/Interchange/Runtime/Source/Import/Private/Mesh/InterchangeGeometryCacheFactory.cpp b/Engine/Plugins/Interchange/Runtime/Source/Import/Private/Mesh/InterchangeGeometryCacheFactory.cpp
--- a/Engine/Plugins/Interchange/Runtime/Source/Import/Private/Mesh/InterchangeGeometryCacheFactory.cpp
+++ b/Engine/Plugins/Interchange/Runtime/Source/Import/Private/Mesh/InterchangeGeometryCacheFactory.cpp
@@ -334,7 +334,7 @@ UInterchangeFactoryBase::FImportAssetResult UInterchangeGeometryCacheFactory::Be
 			GlobalMeshTransform = GlobalOffsetTransform;
 		}
 
-		if (!ensure(GeometryCacheNode || MeshNode))
+		if (!ensure(MeshNode))
 		{
 			UE_LOG(LogInterchangeImport, Warning, TEXT("Invalid mesh reference when importing GeometryCache asset %s."), *AssetName);
 			continue;
