[USD Interchange] Adds kind and custom attribute/metadata support to Nanite assembly traversal

diff --git a/Engine/Plugins/Interchange/Extensions/OpenUSD/Source/Import/Private/InterchangeUsdTranslator.cpp b/Engine/Plugins/Interchange/Extensions/OpenUSD/Source/Import/Private/InterchangeUsdTranslator.cpp
--- a/Engine/Plugins/Interchange/Extensions/OpenUSD/Source/Import/Private/InterchangeUsdTranslator.cpp
+++ b/Engine/Plugins/Interchange/Extensions/OpenUSD/Source/Import/Private/InterchangeUsdTranslator.cpp
@@ -489,6 +489,367 @@
 		NodeUid = Result;
 	}
 
+	template<typename USDType, typename InterchangeType = USDType>
+	InterchangeType CoerceToInterchangeType(const USDType& USDValue)
+	{
+		if constexpr (std::is_same_v<USDType, FMatrix2D> && std::is_same_v<InterchangeType, FMatrix44d>)
+		{
+			FMatrix44d Result{ EForceInit::ForceInitToZero };
+			Result.M[0][0] = USDValue.Row0.X;
+			Result.M[0][1] = USDValue.Row0.Y;
+
+			Result.M[1][0] = USDValue.Row1.X;
+			Result.M[1][1] = USDValue.Row1.Y;
+			return Result;
+		}
+		else if constexpr (std::is_same_v<USDType, FMatrix3D> && std::is_same_v<InterchangeType, FMatrix44d>)
+		{
+			FMatrix44d Result{ EForceInit::ForceInitToZero };
+			Result.M[0][0] = USDValue.Row0.X;
+			Result.M[0][1] = USDValue.Row0.Y;
+			Result.M[0][2] = USDValue.Row0.Z;
+
+			Result.M[1][0] = USDValue.Row1.X;
+			Result.M[1][1] = USDValue.Row1.Y;
+			Result.M[1][2] = USDValue.Row1.Z;
+
+			Result.M[2][0] = USDValue.Row2.X;
+			Result.M[2][1] = USDValue.Row2.Y;
+			Result.M[2][2] = USDValue.Row2.Z;
+			return Result;
+		}
+		else
+		{
+			return (InterchangeType)USDValue;
+		}
+	};
+
+	template<typename USDType, typename InterchangeType = USDType>
+	void AddUserAttribute(const FString& KeyName, const UE::FVtValue& VtValue, UInterchangeSceneNode* Node)
+	{
+		const USDType USDValue = VtValue.Get<USDType>();
+		const TOptional<FString> PayloadKey;
+
+		if constexpr (std::is_same_v<InterchangeType, USDType>)
+		{
+			UInterchangeUserDefinedAttributesAPI::CreateUserDefinedAttribute(Node, KeyName, USDValue, PayloadKey);
+		}
+		else if constexpr (TIsTArray<USDType>::Value)
+		{
+			static_assert(TIsTArray<InterchangeType>::Value);
+
+			InterchangeType InterchangeArray;
+			InterchangeArray.SetNum(USDValue.Num());
+			for (int32 Index = 0; Index < USDValue.Num(); ++Index)
+			{
+				InterchangeArray[Index] = CoerceToInterchangeType<typename USDType::ElementType, typename InterchangeType::ElementType>(USDValue[Index]);
+			}
+
+			UInterchangeUserDefinedAttributesAPI::CreateUserDefinedAttribute(Node, KeyName, InterchangeArray, PayloadKey);
+		}
+		else
+		{
+			const InterchangeType InterchangeValue = CoerceToInterchangeType<USDType, InterchangeType>(USDValue);
+			UInterchangeUserDefinedAttributesAPI::CreateUserDefinedAttribute(Node, KeyName, InterchangeValue, PayloadKey);
+		}
+	};
+
+	// Note: We could have just used the Cpp type names here instead of the value type token, but doing the
+	// latter lets us handle the color types as FLinearColor, which is probably more useful
+	// References:
+	// - https://openusd.org/docs/api/_usd__page__datatypes.html
+	// - EAttributeTypes declaration
+	// - Engine/Source/Runtime/Interchange/Core/Private/Tests/StorageTest.cpp
+	using AddUserAttributeFunc = TFunction<void(const FString&, const UE::FVtValue&, UInterchangeSceneNode*)>;
+	static TMap<FName, AddUserAttributeFunc> AddUserAttributeFuncs = {
+		{TEXT("bool"),  					AddUserAttribute<bool>},
+
+		{TEXT("uchar"), 					AddUserAttribute<uint8>},
+		{TEXT("unsigned char"), 			AddUserAttribute<uint8>},
+		{TEXT("int"), 						AddUserAttribute<int32>},
+		{TEXT("unsigned int"), 				AddUserAttribute<uint32>},
+		{TEXT("uint"), 						AddUserAttribute<uint32>},
+		{TEXT("int64"), 					AddUserAttribute<int64>},
+		{TEXT("__int64"), 					AddUserAttribute<int64>},
+		{TEXT("unsigned __int64"), 			AddUserAttribute<uint64>},
+		{TEXT("uint64"), 					AddUserAttribute<uint64>},
+
+		{TEXT("half"),	 					AddUserAttribute<FFloat16>},
+		{TEXT("pxr_half::half"),	 		AddUserAttribute<FFloat16>},
+		{TEXT("float"), 					AddUserAttribute<float>},
+		{TEXT("double"), 					AddUserAttribute<double>},
+		{TEXT("timecode"), 					AddUserAttribute<FSdfTimeCode, double>},
+		{TEXT("SdfTimeCode"), 				AddUserAttribute<FSdfTimeCode, double>},
+
+		{TEXT("string"), 					AddUserAttribute<FString>},
+		{TEXT("token"), 					AddUserAttribute<FName>},
+		{TEXT("TfToken"), 					AddUserAttribute<FName>},
+		{TEXT("asset"), 					AddUserAttribute<FSdfAssetPath, FString>},
+		{TEXT("SdfAssetPath"), 				AddUserAttribute<FSdfAssetPath, FString>},
+
+		{TEXT("matrix2d"), 					AddUserAttribute<FMatrix2D, FMatrix44d>},
+		{TEXT("GfMatrix2d"), 				AddUserAttribute<FMatrix2D, FMatrix44d>},
+		{TEXT("matrix3d"), 					AddUserAttribute<FMatrix3D, FMatrix44d>},
+		{TEXT("GfMatrix3d"), 				AddUserAttribute<FMatrix3D, FMatrix44d>},
+		{TEXT("matrix4d"), 					AddUserAttribute<FMatrix44d>},
+		{TEXT("GfMatrix4d"), 				AddUserAttribute<FMatrix44d>},
+		{TEXT("frame4d"), 					AddUserAttribute<FMatrix44d>},
+
+		{TEXT("quath"), 					AddUserAttribute<FQuat4h, FQuat4f>},
+		{TEXT("GfQuath"), 					AddUserAttribute<FQuat4h, FQuat4f>},
+		{TEXT("quatf"), 					AddUserAttribute<FQuat4f>},
+		{TEXT("GfQuatf"), 					AddUserAttribute<FQuat4f>},
+		{TEXT("quatd"), 					AddUserAttribute<FQuat4d>},
+		{TEXT("GfQuatd"), 					AddUserAttribute<FQuat4d>},
+
+		{TEXT("half2"), 					AddUserAttribute<FVector2DHalf>},
+		{TEXT("GfVec2h"), 					AddUserAttribute<FVector2DHalf>},
+		{TEXT("float2"), 					AddUserAttribute<FVector2f>},
+		{TEXT("GfVec2f"), 					AddUserAttribute<FVector2f>},
+		{TEXT("double2"), 					AddUserAttribute<FVector2d>},
+		{TEXT("GfVec2d"), 					AddUserAttribute<FVector2d>},
+		{TEXT("int2"), 						AddUserAttribute<FIntPoint>},
+		{TEXT("GfVec2i"), 					AddUserAttribute<FIntPoint>},
+
+		{TEXT("half3"), 					AddUserAttribute<FVector3h, FVector3f>},
+		{TEXT("GfVec3h"), 					AddUserAttribute<FVector3h, FVector3f>},
+		{TEXT("point3h"), 					AddUserAttribute<FVector3h, FVector3f>},
+		{TEXT("normal3h"), 					AddUserAttribute<FVector3h, FVector3f>},
+		{TEXT("vector3h"), 					AddUserAttribute<FVector3h, FVector3f>},
+		{TEXT("color3h"), 					AddUserAttribute<FVector3h, FLinearColor>},
+
+		{TEXT("float3"), 					AddUserAttribute<FVector3f>},
+		{TEXT("GfVec3f"), 					AddUserAttribute<FVector3f>},
+		{TEXT("point3f"), 					AddUserAttribute<FVector3f>},
+		{TEXT("normal3f"), 					AddUserAttribute<FVector3f>},
+		{TEXT("vector3f"), 					AddUserAttribute<FVector3f>},
+		{TEXT("color3f"), 					AddUserAttribute<FVector3f, FLinearColor>},
+
+		{TEXT("double3"), 					AddUserAttribute<FVector3d>},
+		{TEXT("GfVec3d"), 					AddUserAttribute<FVector3d>},
+		{TEXT("point3d"), 					AddUserAttribute<FVector3d>},
+		{TEXT("normal3d"), 					AddUserAttribute<FVector3d>},
+		{TEXT("vector3d"), 					AddUserAttribute<FVector3d>},
+		{TEXT("color3d"), 					AddUserAttribute<FVector3d, FLinearColor>},
+
+		{TEXT("int3"), 						AddUserAttribute<FIntVector>},
+		{TEXT("GfVec3i"), 					AddUserAttribute<FIntVector>},
+
+		{TEXT("half4"), 					AddUserAttribute<FVector4h, FVector4f>},
+		{TEXT("GfVec4h"), 					AddUserAttribute<FVector4h, FVector4f>},
+		{TEXT("float4"), 					AddUserAttribute<FVector4f>},
+		{TEXT("GfVec4f"), 					AddUserAttribute<FVector4f>},
+		{TEXT("double4"), 					AddUserAttribute<FVector4d>},
+		{TEXT("GfVec4d"), 					AddUserAttribute<FVector4d>},
+		{TEXT("int4"), 						AddUserAttribute<FIntRect>},
+		{TEXT("GfVec4i"), 					AddUserAttribute<FIntRect>},
+		{TEXT("color4h"), 					AddUserAttribute<FVector4h, FLinearColor>},
+		{TEXT("color4f"), 					AddUserAttribute<FVector4f, FLinearColor>},
+		{TEXT("color4d"), 					AddUserAttribute<FVector4d, FLinearColor>},
+
+		{TEXT("bool[]"),  					AddUserAttribute<TArray<bool>>},
+		{TEXT("VtArray<bool>"),  			AddUserAttribute<TArray<bool>>},
+
+		{TEXT("uchar[]"), 					AddUserAttribute<TArray<uint8>>},
+		{TEXT("VtArray<unsigned char>"), 	AddUserAttribute<TArray<uint8>>},
+		{TEXT("int[]"), 					AddUserAttribute<TArray<int32>>},
+		{TEXT("VtArray<int>"), 				AddUserAttribute<TArray<int32>>},
+		{TEXT("uint[]"), 					AddUserAttribute<TArray<uint32>>},
+		{TEXT("VtArray<unsigned int>"), 	AddUserAttribute<TArray<uint32>>},
+		{TEXT("int64[]"), 					AddUserAttribute<TArray<int64>>},
+		{TEXT("VtArray<__int64>"), 			AddUserAttribute<TArray<int64>>},
+		{TEXT("uint64[]"), 					AddUserAttribute<TArray<uint64>>},
+		{TEXT("VtArray<unsigned __int64>"), AddUserAttribute<TArray<uint64>>},
+
+		{TEXT("half[]"),	 				AddUserAttribute<TArray<FFloat16>>},
+		{TEXT("VtArray<pxr_half::half>"),	AddUserAttribute<TArray<FFloat16>>},
+		{TEXT("float[]"), 					AddUserAttribute<TArray<float>>},
+		{TEXT("VtArray<float>"), 			AddUserAttribute<TArray<float>>},
+		{TEXT("double[]"), 					AddUserAttribute<TArray<double>>},
+		{TEXT("VtArray<double>"), 			AddUserAttribute<TArray<double>>},
+		{TEXT("timecode[]"), 				AddUserAttribute<TArray<FSdfTimeCode>, TArray<double>>},
+		{TEXT("VtArray<SdfTimeCode>"), 		AddUserAttribute<TArray<FSdfTimeCode>, TArray<double>>},
+
+		{TEXT("string[]"), 					AddUserAttribute<TArray<FString>>},
+		{TEXT("VtArray<string>"), 			AddUserAttribute<TArray<FString>>},
+		{TEXT("token[]"), 					AddUserAttribute<TArray<FName>>},
+		{TEXT("VtArray<TfToken>"), 			AddUserAttribute<TArray<FName>>},
+		{TEXT("asset[]"), 					AddUserAttribute<TArray<FSdfAssetPath>, TArray<FString>>},
+		{TEXT("VtArray<SdfAssetPath>"), 	AddUserAttribute<TArray<FSdfAssetPath>, TArray<FString>>},
+
+		{TEXT("matrix2d[]"), 				AddUserAttribute<TArray<FMatrix2D>, TArray<FMatrix44d>>},
+		{TEXT("VtArray<GfMatrix2d>"), 		AddUserAttribute<TArray<FMatrix2D>, TArray<FMatrix44d>>},
+		{TEXT("matrix3d[]"), 				AddUserAttribute<TArray<FMatrix3D>, TArray<FMatrix44d>>},
+		{TEXT("VtArray<GfMatrix3d>"), 		AddUserAttribute<TArray<FMatrix3D>, TArray<FMatrix44d>>},
+		{TEXT("matrix4d[]"), 				AddUserAttribute<TArray<FMatrix44d>>},
+		{TEXT("VtArray<GfMatrix4d>"), 		AddUserAttribute<TArray<FMatrix44d>>},
+		{TEXT("frame4d[]"), 				AddUserAttribute<TArray<FMatrix44d>>},
+
+		{TEXT("quath[]"), 					AddUserAttribute<TArray<FQuat4h>, TArray<FQuat4f>>},
+		{TEXT("VtArray<GfQuath>"), 			AddUserAttribute<TArray<FQuat4h>, TArray<FQuat4f>>},
+		{TEXT("quatf[]"), 					AddUserAttribute<TArray<FQuat4f>>},
+		{TEXT("VtArray<GfQuatf>"), 			AddUserAttribute<TArray<FQuat4f>>},
+		{TEXT("quatd[]"), 					AddUserAttribute<TArray<FQuat4d>>},
+		{TEXT("VtArray<GfQuatd>"), 			AddUserAttribute<TArray<FQuat4d>>},
+
+		{TEXT("half2[]"), 					AddUserAttribute<TArray<FVector2DHalf>>},
+		{TEXT("VtArray<GfVec2h>"), 			AddUserAttribute<TArray<FVector2DHalf>>},
+		{TEXT("float2[]"), 					AddUserAttribute<TArray<FVector2f>>},
+		{TEXT("VtArray<GfVec2f>"), 			AddUserAttribute<TArray<FVector2f>>},
+		{TEXT("double2[]"), 				AddUserAttribute<TArray<FVector2d>>},
+		{TEXT("VtArray<GfVec2d>"), 			AddUserAttribute<TArray<FVector2d>>},
+		{TEXT("int2[]"), 					AddUserAttribute<TArray<FIntPoint>>},
+		{TEXT("VtArray<GfVec2i>"), 			AddUserAttribute<TArray<FIntPoint>>},
+
+		{TEXT("half3[]"), 					AddUserAttribute<TArray<FVector3h>, TArray<FVector3f>>},
+		{TEXT("VtArray<GfVec3h>"), 			AddUserAttribute<TArray<FVector3h>, TArray<FVector3f>>},
+		{TEXT("point3h[]"), 				AddUserAttribute<TArray<FVector3h>, TArray<FVector3f>>},
+		{TEXT("normal3h[]"), 				AddUserAttribute<TArray<FVector3h>, TArray<FVector3f>>},
+		{TEXT("vector3h[]"), 				AddUserAttribute<TArray<FVector3h>, TArray<FVector3f>>},
+		{TEXT("color3h[]"), 				AddUserAttribute<TArray<FVector3h>, TArray<FLinearColor>>},
+
+		{TEXT("float3[]"), 					AddUserAttribute<TArray<FVector3f>>},
+		{TEXT("VtArray<GfVec3f>"), 			AddUserAttribute<TArray<FVector3f>>},
+		{TEXT("point3f[]"), 				AddUserAttribute<TArray<FVector3f>>},
+		{TEXT("normal3f[]"), 				AddUserAttribute<TArray<FVector3f>>},
+		{TEXT("vector3f[]"), 				AddUserAttribute<TArray<FVector3f>>},
+		{TEXT("color3f[]"), 				AddUserAttribute<TArray<FVector3f>, TArray<FLinearColor>>},
+
+		{TEXT("double3[]"), 				AddUserAttribute<TArray<FVector3d>>},
+		{TEXT("VtArray<GfVec3d>"), 			AddUserAttribute<TArray<FVector3d>>},
+		{TEXT("point3d[]"), 				AddUserAttribute<TArray<FVector3d>>},
+		{TEXT("normal3d[]"), 				AddUserAttribute<TArray<FVector3d>>},
+		{TEXT("vector3d[]"), 				AddUserAttribute<TArray<FVector3d>>},
+		{TEXT("color3d[]"), 				AddUserAttribute<TArray<FVector3d>, TArray<FLinearColor>>},
+
+		{TEXT("int3[]"), 					AddUserAttribute<TArray<FIntVector>>},
+		{TEXT("VtArray<GfVec3i>"), 			AddUserAttribute<TArray<FIntVector>>},
+
+		{TEXT("half4[]"), 					AddUserAttribute<TArray<FVector4h>, TArray<FVector4f>>},
+		{TEXT("VtArray<GfVec4h>"), 			AddUserAttribute<TArray<FVector4h>, TArray<FVector4f>>},
+		{TEXT("float4[]"), 					AddUserAttribute<TArray<FVector4f>>},
+		{TEXT("VtArray<GfVec4f>"), 			AddUserAttribute<TArray<FVector4f>>},
+		{TEXT("double4[]"), 				AddUserAttribute<TArray<FVector4d>>},
+		{TEXT("VtArray<GfVec4d>"), 			AddUserAttribute<TArray<FVector4d>>},
+		{TEXT("int4[]"), 					AddUserAttribute<TArray<FIntRect>>},
+		{TEXT("VtArray<GfVec4i>"), 			AddUserAttribute<TArray<FIntRect>>},
+		{TEXT("color4h[]"), 				AddUserAttribute<TArray<FVector4h>, TArray<FLinearColor>>},
+		{TEXT("color4f[]"), 				AddUserAttribute<TArray<FVector4f>, TArray<FLinearColor>>},
+		{TEXT("color4d[]"), 				AddUserAttribute<TArray<FVector4d>, TArray<FLinearColor>>},
+	};
+
+	void TranslateAttributes(const UE::FUsdPrim& Prim, UInterchangeSceneNode* Node, const FString& AllowedAttributeRegex)
+	{
+		TRACE_CPUPROFILER_EVENT_SCOPE(TranslateAttributes)
+
+			FRegexPattern RegexPattern{ AllowedAttributeRegex };
+
+		for (const UE::FUsdAttribute& Attr : Prim.GetAttributes())
+		{
+			if (!Attr.HasAuthoredValue())
+			{
+				continue;
+			}
+
+			const FString AttrName = Attr.GetName().ToString();
+
+			FRegexMatcher RegexMatcher{ RegexPattern, AttrName };
+			if (!RegexMatcher.FindNext())
+			{
+				continue;
+			}
+
+			if (AddUserAttributeFunc Func = AddUserAttributeFuncs.FindRef(Attr.GetTypeName()))
+			{
+				UE::FVtValue Value;
+				if (Attr.Get(Value)) // Always check for an opinion on the default time code
+				{
+					Func(AttrName, Value, Node);
+				}
+			}
+		}
+	}
+
+	void TranslateMetadata(const UE::FUsdPrim& Prim, UInterchangeSceneNode* Node, const FString& AllowedMetadataRegex)
+	{
+		TRACE_CPUPROFILER_EVENT_SCOPE(TranslateMetadata)
+
+			FRegexPattern RegexPattern{ AllowedMetadataRegex };
+
+		TFunction<void(const TMap<FString, UE::FVtValue>&, const FString&)> ParseMetadataMap;
+		ParseMetadataMap = [&ParseMetadataMap, &RegexPattern, Node](const TMap<FString, UE::FVtValue>& MetadataMap, const FString& ConcatParentKey)
+		{
+			for (const TPair<FString, UE::FVtValue>& Pair : MetadataMap)
+			{
+				const FString FullKey = ConcatParentKey.IsEmpty() ? Pair.Key : FString::Printf(TEXT("%s:%s"), *ConcatParentKey, *Pair.Key);
+				const FString TypeName = Pair.Value.GetTypeName();
+
+				if (TypeName == TEXT("VtDictionary"))
+				{
+					TMap<FString, UE::FVtValue> NestedMap = Pair.Value.Get<TMap<FString, UE::FVtValue>>();
+					ParseMetadataMap(NestedMap, FullKey);
+				}
+				else if (AddUserAttributeFunc Func = AddUserAttributeFuncs.FindRef(*TypeName))
+				{
+					// Only match when parsing leaf keys, because a key may not match just the parent nested map key directly and
+					// would have otherwise prevented us from stepping into it (e.g. the regex "customData:my" will not match 
+					// "customData", but it would match "customData:myNestedMap")
+					FRegexMatcher RegexMatcher{ RegexPattern, FullKey };
+					if (!RegexMatcher.FindNext())
+					{
+						continue;
+					}
+
+					Func(FullKey, Pair.Value, Node);
+				}
+			}
+		};
+
+		const FString ConcatParentKey;
+		ParseMetadataMap(Prim.GetAllAuthoredMetadata(), ConcatParentKey);
+	}
+
+	UInterchangeSceneNode* CreateSceneNode(
+		const UE::FUsdPrim& Prim, 
+		UInterchangeUsdTranslatorSettings* TranslatorSettings,
+		UInterchangeBaseNodeContainer& NodeContainer, 
+		const FString& SceneNodeUid, 
+		const FString& DisplayLabel,
+		const FString& ParentNodeUid
+	)
+	{
+		UInterchangeSceneNode* SceneNode = NewObject<UInterchangeSceneNode>(&NodeContainer);
+		NodeContainer.SetupNode(SceneNode, SceneNodeUid, DisplayLabel, EInterchangeNodeContainerType::TranslatedScene, ParentNodeUid);
+
+		// Store our purpose here to be filtered on the pipeline
+		const EUsdPurpose PrimPurpose = IUsdPrim::GetPurpose(Prim);
+		if (PrimPurpose != EUsdPurpose::Default)
+		{
+			SceneNode->AddInt32Attribute(UE::Interchange::USD::GeometryPurposeIdentifier, (int32)PrimPurpose);
+		}
+
+		// Store our prim kind as well, if we have any (this becomes the empty string if the prim has no authored kind)
+		const FString KindString = UsdToUnreal::ConvertToken(IUsdPrim::GetKind(Prim));
+		if (!KindString.IsEmpty())
+		{
+			const FString UserDefinedAttributeName = TEXT("kind");
+			TOptional<FString> PayloadKey;
+			UInterchangeUserDefinedAttributesAPI::CreateUserDefinedAttribute(SceneNode, UserDefinedAttributeName, KindString, PayloadKey);
+		}
+
+		if (TranslatorSettings->bTranslatePrimAttributes)
+		{
+			TranslateAttributes(Prim, SceneNode, TranslatorSettings->AttributeRegexFilter);
+		}
+
+		if (TranslatorSettings->bTranslatePrimMetadata)
+		{
+			TranslateMetadata(Prim, SceneNode, TranslatorSettings->MetadataRegexFilter);
+		}
+
+		return SceneNode;
+	}
+
 	bool IsValidLODName(const FString& PrimName)
 	{
 		return PrimName.Len() > LODString.Len() && PrimName.StartsWith(LODString) && PrimName.RightChop(LODString.Len()).IsNumeric();
@@ -1547,666 +1908,3 @@
 			UInterchangeUSDTranslatorImpl& InTranslatorImpl,
 			UInterchangeBaseNodeContainer& InNodeContainer,
 			UInterchangeUsdTranslatorSettings* InTranslatorSettings,
-			UInterchangeSceneNode* InParentSceneNode)
-			: TranslatorImpl(InTranslatorImpl)
-			, NodeContainer(InNodeContainer)
-			, TranslatorSettings(InTranslatorSettings)
-			, ParentSceneNode(InParentSceneNode)
-			, ParentSceneNodeUid(ParentSceneNode->GetUniqueID())
-		{
-
-		}
-
-	private:
-
-		bool MeshTypeIs(ENaniteAssemblyMeshType MeshType) const
-		{
-			return MeshType == AssemblyMeshType;
-		}
-
-		bool IsCachableSkeletalMeshPrim(const UE::FUsdPrim& Prim, EMeshCategory Category)
-		{
-			return 
-				MeshTypeIs(ENaniteAssemblyMeshType::SkeletalMesh)
-				&& Category == EMeshCategory::Part
-				&& Prim.IsInstanceProxy();
-		}
-
-		bool GetCachedSkeletalMeshUid(const UE::FUsdPrim& Prim, EMeshCategory Category, FString& OutMeshUid)
-		{
-			if (IsCachableSkeletalMeshPrim(Prim, Category))
-			{
-				UE::FUsdPrim PrototypePrim = UsdUtils::GetPrototypePrim(Prim);
-				if (const FString* CachedMeshUid = ImplicitPrototypePrimPathToCachedMeshUid.Find(PrototypePrim.GetPrimPath()))
-				{
-					OutMeshUid = *CachedMeshUid;
-					return true;
-				}
-			}
-			return false;
-		}
-
-		void SetCachedSkeletalMeshUid(const UE::FUsdPrim& Prim, EMeshCategory Category, const FString& MeshUid)
-		{
-			if (IsCachableSkeletalMeshPrim(Prim, Category))
-			{
-				UE::FUsdPrim PrototypePrim = UsdUtils::GetPrototypePrim(Prim);
-				ImplicitPrototypePrimPathToCachedMeshUid.Add(PrototypePrim.GetPrimPath(), MeshUid);
-			}
-		}
-
-		FString CreateMeshNode(const UE::FSdfPath& Path, EMeshCategory Category, const TArray<FString>& AssemblyDependencies = {})
-		{
-			FTraversalInfo* Info = TraversalInfoByCachedPrimPath.Find(Path);
-			if (!ensure(Info))
-			{
-				return {};
-			}
-
-			const UE::FUsdPrim Prim = AssemblyRootPrim.GetStage().GetPrimAtPath(Path);
-			if (!ensure(Prim.IsA(TEXT("Mesh"))))
-			{
-				return {};
-			}
-
-			// If this is a repeated instanced skeletal mesh part, reuse the previously generated uid. This is
-			// a workaround for the current lack of instanced skeletal mesh support in Interchange.
-			if (FString CachedMeshUid; GetCachedSkeletalMeshUid(Prim, Category, CachedMeshUid))
-			{
-				return CachedMeshUid;
-			}
-
-			// If the mesh node we're creating has assembly dependencies, it must be the first base mesh, and so 
-            // attach the traversal result to our traversal info for retrieval during payload creation.
-			if (AssemblyDependencies.Num() > 0 && ensure(Result.IsValid()))
-			{
-				Info->NaniteAssemblyTraversalResult = Result;
-			}
-			
-			FString MeshUid = TranslatorImpl.AddMeshNode(Prim, TranslatorSettings, NodeContainer, *Info, false, AssemblyDependencies);
-
-			const FString SceneNodeUid = MakeNodeUid(Prim.GetPrimPath().GetString());
-			const FString DisplayLabel(Prim.GetName().ToString());
-
-			UInterchangeSceneNode* SceneNode = NewObject<UInterchangeSceneNode>(&NodeContainer);
-			static const FString NoParentUid = TEXT("");
-			NodeContainer.SetupNode(SceneNode, SceneNodeUid, DisplayLabel, EInterchangeNodeContainerType::TranslatedScene, NoParentUid);
-
-			constexpr bool bResetCache = false;
-			SceneNode->SetCustomLocalTransform(&NodeContainer, FTransform::Identity, bResetCache);
-			SceneNode->SetCustomAssetInstanceUid(MeshUid);
-
-			// TODO: TranslateAttributes, TranslateMetadata to scene node?
-
-			// Store the mesh uid for the instanced skeletal mesh case
-			SetCachedSkeletalMeshUid(Prim, Category, MeshUid);
-
-			return MeshUid;
-		}
-
-		UE::FSdfPath GetSkelBindingPrimPathForPart(const UE::FUsdPrim& Prim, const FTraversalInfo& Info)
-		{
-			// This should only be called for skeletal mesh parts.
-			// 
-			// Check and get the path to the (possibly ancestor) prim holding skel binding data
-
-			UE::FSdfPath EmptyPath;
-			if (!MeshTypeIs(ENaniteAssemblyMeshType::SkeletalMesh))
-			{
-				return EmptyPath;
-			}
-
-			if (Info.ClosestNaniteAssemblySkelBindingPath)
-			{
-				const UE::FSdfPath SkelBindingPrimPath = UE::FSdfPath(**Info.ClosestNaniteAssemblySkelBindingPath);
-				if (SkelBindingPrimPath.IsPrimPath())
-				{
-					return SkelBindingPrimPath;
-				}
-			}
-
-			InvalidPartMeshNoSkelBindingPaths.Add(Prim.GetPrimPath());
-
-			return EmptyPath;
-		}
-
-		EMeshCategory GetMeshCategory(const UE::FUsdPrim& Prim, const FTraversalInfo& Info, TArray<FPrimPrototypeEntry>& PrototypeStack)
-		{
-			if (!Prim.IsA(TEXT("Mesh")))
-			{
-				return EMeshCategory::None;
-			}
-
-			// A mesh in a skeletal assembly must:
-			//  i: have a bound skeleton 
-			// ii: if it's a 'base' mesh (i.e. connected to the designated assembly root skeleton)
-			//      - must NOT be inside a pointinstancer
-			//      - must NOT be USD native instanced
-			//     or, if it's a 'part' mesh (i.e. attached to some other skeleton)
-			//      - may be inside a pointinstancer, but, not nested, or
-			//      - must be USD native instanced (because we don't want to encourage users to create tons of unique non-instanced parts) 
-			//
-			if (MeshTypeIs(ENaniteAssemblyMeshType::SkeletalMesh) && Info.BoundSkeletonPrimPath)
-			{
-				if (BaseSkeletonPrimPath.GetString() == **Info.BoundSkeletonPrimPath) // base skeleton mesh
-				{
-					if (PrototypeStack.IsEmpty() && !Prim.IsInstanceProxy())
-					{
-						return EMeshCategory::Base;
-					}
-				}
-				else // part skeleton mesh
-				{
-					if (PrototypeStack.Num() == 1 || Prim.IsInstanceProxy())
-					{
-						return EMeshCategory::Part;
-					}
-				}
-			}
-
-			// A part mesh in a static assembly must 
-			//  i: not have skel binding data
-			// ii: must be inside a pointinstancer prototype, or, be USD native instanced.
-			//
-			else if (MeshTypeIs(ENaniteAssemblyMeshType::StaticMesh))
-			{
-				if (!Prim.HasAPI(TEXT("SkelBindingAPI")) && (PrototypeStack.Num() > 0 || Prim.IsInstanceProxy()))
-				{
-					return EMeshCategory::Part;
-				}
-			}
-
-			// TODO: communicate all of the above to the user, somehow
-			return EMeshCategory::None;
-		}
-
-		void HandleOrIgnoreMesh(
-			const UE::FUsdPrim& Prim, 
-			FTraversalInfo Info, 
-			TArray<FTransform>& TransformStack, 
-			TArray<FPrimPrototypeEntry>& PrototypeStack)
-		{
-			if (!Prim.IsA(TEXT("Mesh")) || !Result.IsValid())
-			{
-				return;
-			}
-
-			EMeshCategory Category = GetMeshCategory(Prim, Info, PrototypeStack);
-			if (Category == EMeshCategory::None)
-			{
-				return;
-			}
-
-			// If this is a skeletal mesh part, check we have a prim with binding data
-			UE::FSdfPath SkelBindingPrimPath;
-			if (MeshTypeIs(ENaniteAssemblyMeshType::SkeletalMesh) && Category == EMeshCategory::Part)
-			{
-				SkelBindingPrimPath = GetSkelBindingPrimPathForPart(Prim, Info);
-				if (!SkelBindingPrimPath.IsPrimPath())
-				{
-					return;
-				}
-			}
-
-			ensure(
-				Result->AddEntry(MakeShared<FMeshEntry>(
-					Category, 
-					Prim.GetPrimPath(), 
-					SkelBindingPrimPath, 
-					TransformStack, 
-					PrototypeStack)));
-
-			TraversalInfoByCachedPrimPath.Add(Prim.GetPrimPath(), Info);
-		}
-
-		void HandleExternalAssetReference(
-			const UE::FUsdPrim& Prim,
-			FTraversalInfo Info,
-			TArray<FTransform>& TransformStack,
-			TArray<FPrimPrototypeEntry>& PrototypeStack,
-			const FString& AssetReference)
-		{
-			if (!Result.IsValid())
-			{
-				return;
-			}
-
-			UE::FSdfPath SkelBindingPrimPath;
-
-			// External asset references are only supported on parts.
-			EMeshCategory Category = EMeshCategory::Part;
-
-			if (MeshTypeIs(ENaniteAssemblyMeshType::SkeletalMesh))
-			{
-				// If we found the schema applied to a descendant of the base skelroot prim
-				// we should skip as we're not a part.
-				if (Prim.GetPrimPath().HasPrefix(BaseSkelRootPrimPath))
-				{
-					return;
-				}
-
-				// This is a skeletal mesh part - check we have prim with binding data
-				SkelBindingPrimPath = GetSkelBindingPrimPathForPart(Prim, Info);
-				if (!SkelBindingPrimPath.IsPrimPath())
-				{
-					return;
-				}
-			}
-
-			TSharedPtr<FMeshEntry> Entry = MakeShared<FMeshEntry>(
-				Category,
-				Prim.GetPrimPath(),
-				SkelBindingPrimPath,
-				TransformStack,
-				PrototypeStack);
-
-			if (!ensure(Result->AddEntry(Entry)))
-			{
-				return;
-			}
-			
-			Entry->NodeUid = AssetReference;
-
-			TraversalInfoByCachedPrimPath.Add(Prim.GetPrimPath(), Info);
-		}
-
-		void HandleSkeleton(const UE::FUsdPrim& Prim, FTraversalInfo Info, TArray<FTransform>& TransformStack) const
-		{
-			const bool bIsBaseSkeleton = Prim.GetPrimPath() == BaseSkeletonPrimPath;
-
-			const FString SceneNodeUid = MakeNodeUid(Prim.GetPrimPath().GetString());
-			const FString DisplayLabel(Prim.GetName().ToString());
-
-			UInterchangeSceneNode* SceneNode = NewObject<UInterchangeSceneNode>(&NodeContainer);
-			FString ThisParentSceneUid = bIsBaseSkeleton ? ParentSceneNodeUid : TEXT("");
-			NodeContainer.SetupNode(SceneNode, SceneNodeUid, DisplayLabel, EInterchangeNodeContainerType::TranslatedScene, ThisParentSceneUid);
-
-			// TODO: TranslateAttributes, TranslateMetadata to scene node?
-
-			const FTransform Transform = bIsBaseSkeleton ? TransformStack.Last() : FTransform::Identity;
-			constexpr bool bResetCache = false;
-			SceneNode->SetCustomLocalTransform(&NodeContainer, Transform, bResetCache);
-
-			TranslatorImpl.AddSkeletonNodes(Prim, *SceneNode, NodeContainer, Info);
-		}
-
-		bool FindAndSetBaseSkeleton()
-		{
-			if (!MeshTypeIs(ENaniteAssemblyMeshType::SkeletalMesh))
-			{
-				return true;
-			}
-
-			// Get the "base" skeleton prim as specified by the schemas "unreal:naniteAssembly:skeleton" relationship
-			
-			const UE::FUsdPrim SkelPrim = UsdToUnreal::NaniteAssemblyUtils::GetBaseSkeletonPrimForSkeletalMeshAssembly(AssemblyRootPrim);
-			if (!SkelPrim)
-			{
-				// Note reason for failure will have been logged by utility call above.
-				return false;
-			}
-			BaseSkeletonPrimPath = SkelPrim.GetPrimPath();
-
-			// We got the skeleton prim, so now check we have an ancestor SkelRoot prim before the assembly root prim in order to try
-			// detect weird or broken hierarchy configurations.
-
-			UE::FUsdPrim CurrentPrim = SkelPrim;
-			while (!CurrentPrim.IsPseudoRoot())
-			{
-				if (CurrentPrim.IsA(TEXT("SkelRoot")))
-				{
-					BaseSkelRootPrimPath = CurrentPrim.GetPrimPath();
-					break;
-				}
-				CurrentPrim = CurrentPrim.GetParent();
-			}
-
-			if (BaseSkelRootPrimPath.IsEmpty())
-			{
-				UE_LOG(LogUsd, Warning, 
-					TEXT("Failed to find SkelRoot prim for base Skeleton prim '%s' while processing NaniteAssemblyRoot prim '%s'")
-					, *BaseSkeletonPrimPath.GetString()
-					, *AssemblyRootPrim.GetPrimPath().GetString()
-				);
-				return false;
-			}
-			
-			return true;
-		}
-
-		bool ComputeTraversalResult(FTraversalInfo Info, const FTransform& ParentGlobalTransform)
-		{
-			Result = MakeShared<FNaniteAssemblyTraversalResult>(AssemblyMeshType, AssemblyRootPrim.GetPrimPath());
-
-			TArray<FTransform> TransformStack = { ParentGlobalTransform };
-			TArray<FPrimPrototypeEntry> PrototypeStack;
-			constexpr bool bTraverseInstanceProxies = true;
-			for (const UE::FUsdPrim& ChildPrim : AssemblyRootPrim.GetFilteredChildren(bTraverseInstanceProxies))
-			{
-				TraverseImpl(ChildPrim, Info, TransformStack, PrototypeStack);
-			}
-
-			if (!ValidateTraversalResult())
-			{
-				Result.Reset();
-				return false;
-			}
-
-			return true;
-		}
-
-		void TraversePointInstancerHierarchy(
-			const UE::FUsdPrim& Prim,
-			FTraversalInfo Info,
-			TArray<FTransform> TransformStack,
-			TArray<FPrimPrototypeEntry>& PrototypeStack
-		)
-		{
-			const TArray<UE::FSdfPath> PrototypePaths = UsdToUnreal::NaniteAssemblyUtils::GetPointInstancerPrototypePaths(Prim);
-			if (PrototypePaths.IsEmpty())
-			{
-				return;
-			}
-
-			for (int32 PrototypeIndex = 0; PrototypeIndex < PrototypePaths.Num(); ++PrototypeIndex)
-			{
-				UE::FUsdPrim PrototypePrim = Prim.GetStage().GetPrimAtPath(PrototypePaths[PrototypeIndex]);
-				if (!PrototypePrim)
-				{
-					continue;
-				}
-
-				FTraversalInfo PrototypePrimInfo = Info;
-				PrototypePrimInfo.UpdateWithCurrentPrim(PrototypePrim);
-
-				// Push transforms and prototype before entering each pointinstancer prototype
-				TransformStack.Add(FTransform::Identity);
-				PrototypeStack.Emplace(Prim.GetPrimPath(), PrototypePrim.GetPrimPath(), PrototypeIndex, PrototypePaths.Num());
-
-				TraverseImpl(PrototypePrim, Info, TransformStack, PrototypeStack);
-
-				// Pop before entering the next
-				PrototypeStack.Pop();
-				TransformStack.Pop();
-			}
-		}
-
-		void TraverseImpl(const UE::FUsdPrim& Prim, FTraversalInfo Info, TArray<FTransform> TransformStack, TArray<FPrimPrototypeEntry>& PrototypeStack)
-		{
-			if (!ensure(TransformStack.Num() > 0))
-			{
-				return;
-			}
-
-			if (Prim.HasAPI(TEXT("NaniteAssemblyRootAPI")))
-			{
-				UE_LOG(LogUsd, Warning,
-					TEXT("Skipping nested NaniteAssemblyRootAPI prim '%s' found below current root '%s'")
-					, *Prim.GetPrimPath().GetString()
-					, *AssemblyRootPrim.GetPrimPath().GetString()
-				);
-				return;
-			}
-
-			if (Prim.IsA(TEXT("Material")))
-			{
-				return;
-			}
-
-			Info.UpdateWithCurrentPrim(Prim);
-
-			if (Prim.IsA(TEXT("Xformable")))
-			{
-				FTransform LocalTransform = FTransform::Identity;
-				bool bResetTransformStack = false;
-				UsdToUnreal::ConvertXformable(Prim.GetStage(), UE::FUsdTyped(Prim), LocalTransform, UsdUtils::GetEarliestTimeCode(), &bResetTransformStack);
-				
-				FTransform& CurrentTransform = TransformStack.Last();
-				CurrentTransform = LocalTransform * CurrentTransform;
-			}
-
-			if (FString AssetRef; UsdToUnreal::NaniteAssemblyUtils::GetExternalAssetReference(Prim, AssetRef))
-			{
-				HandleExternalAssetReference(Prim, Info, TransformStack, PrototypeStack, AssetRef);
-				// The user is mostly likely expecting the subtree below the prim with the external
-				// reference applied to be pruned, so we stop the traversal.
-				return;
-			}
-
-			if (Prim.IsA(TEXT("Mesh")))
-			{
-				HandleOrIgnoreMesh(Prim, Info, TransformStack, PrototypeStack);
-			}
-			else if (Prim.IsA(TEXT("PointInstancer")))
-			{
-				TraversePointInstancerHierarchy(Prim, Info, TransformStack, PrototypeStack);
-			}
-			else if (Prim.IsA(TEXT("Skeleton")))
-			{
-				HandleSkeleton(Prim, Info, TransformStack);
-			}
-			else
-			{
-				constexpr bool bTraverseInstanceProxies = true;
-				for (const UE::FUsdPrim& ChildPrim : Prim.GetFilteredChildren(bTraverseInstanceProxies))
-				{
-					TraverseImpl(ChildPrim, Info, TransformStack, PrototypeStack);
-				}
-			}
-		}
-
-		bool ValidateHierarchy(const FTraversalInfo& Info) const
-		{
-			if (AssemblyMeshType == ENaniteAssemblyMeshType::None)
-			{
-				return false;
-			}
-
-			// Nanite assemblies cannot be inside LODs.
-			if (Info.bInsideLOD || Info.bIsLODContainer)
-			{
-				UE_LOG(LogUsd, Warning,
-					TEXT("Ignoring NaniteAssemblyRootAPI prim '%s' which is inside an LOD (not supported).")
-					, *AssemblyRootPrim.GetPrimPath().GetString()
-				);
-				return false;
-			}
-
-			// Check ancestors for configurations that we know will cause problems:
-			
-			UE::FUsdPrim CurrentPrim = AssemblyRootPrim;
-
-			while (!CurrentPrim.IsPseudoRoot())
-			{
-				if (CurrentPrim.IsA(TEXT("PointInstancer")))
-				{
-					UE_LOG(LogUsd, Warning,
-						TEXT("NaniteAssemblyRoot '%s' cannot currently be a descendant of a pointinstancer - '%s'")
-						, *AssemblyRootPrim.GetPrimPath().GetString()
-						, *CurrentPrim.GetPrimPath().GetString()
-					);
-					return false;
-				}
-				else if (CurrentPrim.HasAPI(*UsdToUnreal::ConvertToken(UnrealIdentifiers::LodSubtreeAPI)))
-				{
-					UE_LOG(LogUsd, Warning,
-						TEXT("NaniteAssemblyRoot '%s' cannot be nested inside an %s container.")
-						, *AssemblyRootPrim.GetPrimPath().GetString()
-						, *UsdToUnreal::ConvertToken(UnrealIdentifiers::LodSubtreeAPI)
-					);
-					return false;
-				}
-				CurrentPrim = CurrentPrim.GetParent();
-			}
-
-			return true;
-		}
-
-		bool ValidateTraversalResult() const
-		{
-			if (!Result)
-			{
-				return false;
-			}
-
-			// Report skel mesh parts that were dropped due to missing bindings
-			if (InvalidPartMeshNoSkelBindingPaths.Num() > 0)
-			{
-				FString Message = FString::Printf(
-					TEXT("Ignoring (%d) skeletal mesh prims missing joint binding data while generating parts for Nanite assembly ")
-					TEXT("prim '%s'.\nPlease check that the %s schema is applied to an ancestor prim and valid primvar data ")
-					TEXT("authored for attributes '%s' and '%s'...\n")
-					, InvalidPartMeshNoSkelBindingPaths.Num()
-					, *AssemblyRootPrim.GetPrimPath().GetString()
-					, *UsdToUnreal::ConvertToken(UnrealIdentifiers::NaniteAssemblySkelBindingAPI)
-					, *UsdToUnreal::ConvertToken(UnrealIdentifiers::PrimvarsUnrealNaniteAssemblyBindJoints)
-					, *UsdToUnreal::ConvertToken(UnrealIdentifiers::PrimvarsUnrealNaniteAssemblyBindJointWeights)
-				);
-				for (int32 Index = 0; Index < InvalidPartMeshNoSkelBindingPaths.Num(); ++Index)
-				{
-					// there might be a massive number of prims here, so only print the first few
-					int32 constexpr MaxEntries = 50;
-					if (Index < MaxEntries)
-					{
-						Message.Append(FString::Printf(TEXT("... %s\n"), *InvalidPartMeshNoSkelBindingPaths[Index].GetString()));
-					}
-					else
-					{
-						Message.Append(FString::Printf(TEXT("... (skipped %d more)"), InvalidPartMeshNoSkelBindingPaths.Num() - MaxEntries));
-						break;
-					}
-				}
-
-				UE_LOG(LogUsd, Warning, TEXT("%s"), *Message);
-			}
-
-			// Check we found some parts, and, if this is a skeletal mesh assembly, we also require at least one base mesh
-			const bool bHasBaseMeshes = Result->HasEntriesForCategory(EMeshCategory::Base);
-			const bool bHasPartMeshes = Result->HasEntriesForCategory(EMeshCategory::Part);
-			if (!bHasPartMeshes || (MeshTypeIs(ENaniteAssemblyMeshType::SkeletalMesh) && !bHasBaseMeshes))
-			{
-				UE_LOG(LogUsd, Warning,
-					TEXT("Failed to find any valid base/part meshes for NaniteAssemblyRootAPI prim '%s'")
-					, *AssemblyRootPrim.GetPrimPath().GetString()
-				);
-
-				return false;
-			}
-
-			return true;
-		}
-
-	public:
-
-		void Traverse(const UE::FUsdPrim& InAssemblyRootPrim, FTraversalInfo Info, const FTransform& ParentGlobalTransform)
-		{
-			AssemblyRootPrim = InAssemblyRootPrim;
-
-			if (!AssemblyRootPrim.HasAPI(TEXT("NaniteAssemblyRootAPI")))
-			{
-				return;
-			}
-
-			// Setup and initial validation
-
-			AssemblyMeshType = UsdToUnreal::NaniteAssemblyUtils::GetNaniteAssemblyMeshType(AssemblyRootPrim);
-
-			if (!ValidateHierarchy(Info) || !FindAndSetBaseSkeleton())
-			{
-				return;
-			}
-
-			// Compute the traversal result that will be attached to the traversal info object of
-			// the skeletal or static assembly base mesh. 
-			if (!ComputeTraversalResult(Info, ParentGlobalTransform) || !Result.IsValid())
-			{
-				return;
-			}
-			
-			// Now create mesh and scene nodes for parts. Do these first because we will need the resulting node 
-			// uids when creating the base meshes in a moment, in order to setup the correct dependencies.
-			for (TSharedPtr<FMeshEntry>& Entry : Result->GetEntriesForCategory(EMeshCategory::Part))
-			{
-				// Note that external asset ref part entries will have already set the uid to whatever asset.
-				if (Entry->NodeUid.IsEmpty()) 
-				{
-					// If create mesh fails here we'll wind up with an invalid/empty node uid for the entry.
-					// That's ok however because the payload processing will validate all entries anyway.
-					Entry->NodeUid = CreateMeshNode(Entry->PrimPath, Entry->Category);
-				}
-			}
-
-			// Skeletal mesh - now create base meshes and register dependencies on the parts we created
-			// above. Note also that we only need to apply the dependencies to one of the base meshes
-			// since Interchange will currently combine all the meshes anyway.
-			if (MeshTypeIs(ENaniteAssemblyMeshType::SkeletalMesh))
-			{
-				static const TArray<FString> EmptyDependencyArray;
-				bool bIsFirst = true;
-				for (TSharedPtr<FMeshEntry>& Entry : Result->GetEntriesForCategory(EMeshCategory::Base))
-				{
-					const TArray<FString> PartDependencies = bIsFirst 
-						? Result->GetMeshUidsForCategory(EMeshCategory::Part) 
-						: EmptyDependencyArray;
-
-					Entry->NodeUid = CreateMeshNode(Entry->PrimPath, Entry->Category, PartDependencies);
-
-					bIsFirst = false;
-				}
-			}
-			// Static mesh - the NaniteAssemblyRootAPI schema doesn't currently support specifying a base 
-			// mesh when in static mesh mode. So here we create a dummy mesh node that will result in an 
-			// empty mesh that the mesh factory can eventually add the assembly data to. 
-			else if (MeshTypeIs(ENaniteAssemblyMeshType::StaticMesh))
-			{
-				const FString MeshUid = MakeNodeUid(MeshPrefix + AssemblyRootPrim.GetPrimPath().GetString());
-				const FString NodeName(AssemblyRootPrim.GetPrimPath().GetName());
-
-				UInterchangeMeshNode* MeshNode = NewObject<UInterchangeMeshNode>(&NodeContainer);
-				NodeContainer.SetupNode(MeshNode, MeshUid, NodeName, EInterchangeNodeContainerType::TranslatedAsset);
-				MeshNode->SetAssetName(NodeName);
-				MeshNode->SetPayLoadKey(AssemblyRootPrim.GetPrimPath().GetString(), EInterchangeMeshPayLoadType::STATIC);
-
-				for (const FString& PartUid : Result->GetMeshUidsForCategory(EMeshCategory::Part))
-				{
-					MeshNode->AddAssemblyPartDependencyUid(PartUid);
-				}
-
-				ParentSceneNode->SetCustomAssetInstanceUid(MeshUid);
-
-				// Store the traversal result and cache the info for the mesh
-				Info.NaniteAssemblyTraversalResult = Result;
-				FWriteScopeLock ScopedInfoWriteLock{ TranslatorImpl.CachedTraversalInfoLock };
-				TranslatorImpl.NodeUidToCachedTraversalInfo.Add(MeshUid, Info);
-			}
-		}
-	};
-
-	void UInterchangeUSDTranslatorImpl::AddNaniteAssemblyHierarchy(
-		const UE::FUsdPrim& Prim,
-		UInterchangeUsdTranslatorSettings* TranslatorSettings,
-		UInterchangeBaseNodeContainer& NodeContainer,
-		UInterchangeSceneNode* ParentSceneNode,
-		FTraversalInfo& Info)
-	{
-		FNaniteAssemblyTraversalHelper NaniteAssemblyTraversalHelper(*this, NodeContainer, TranslatorSettings, ParentSceneNode);
-		NaniteAssemblyTraversalHelper.Traverse(Prim, Info, FTransform());
-	}
-
-	//Helper to traverse PointInstancer subhierarchy and process it into InstancedStaticMeshComponents (only).
-	struct FISMComponentTraversalHelper
-	{
-		UInterchangeUSDTranslatorImpl& TranslatorImpl;
-		UInterchangeBaseNodeContainer& NodeContainer;
-		UInterchangeUsdTranslatorSettings* TranslatorSettings;
-		UInterchangeSceneNode* ParentSceneNode;
-		const FString ParentSceneNodeUid;
-
-		//While we could lookup from the NodeContainer, 
-		// we are implementing a LookUp table specifically for ISMComponentNodes for speed:
-		TMap<FString, UInterchangeInstancedStaticMeshComponentNode*> UidToComponentNode;
-
-		FISMComponentTraversalHelper(
-			UInterchangeUSDTranslatorImpl& InTranslatorImpl,
-			UInterchangeBaseNodeContainer& InNodeContainer,
-			UInterchangeUsdTranslatorSettings* InTranslatorSettings,
@@ -2220,6 +1918,662 @@
 
 		}
 
+	private:
+
+		bool MeshTypeIs(ENaniteAssemblyMeshType MeshType) const
+		{
+			return MeshType == AssemblyMeshType;
+		}
+
+		bool IsCachableSkeletalMeshPrim(const UE::FUsdPrim& Prim, EMeshCategory Category)
+		{
+			return 
+				MeshTypeIs(ENaniteAssemblyMeshType::SkeletalMesh)
+				&& Category == EMeshCategory::Part
+				&& Prim.IsInstanceProxy();
+		}
+
+		bool GetCachedSkeletalMeshUid(const UE::FUsdPrim& Prim, EMeshCategory Category, FString& OutMeshUid)
+		{
+			if (IsCachableSkeletalMeshPrim(Prim, Category))
+			{
+				UE::FUsdPrim PrototypePrim = UsdUtils::GetPrototypePrim(Prim);
+				if (const FString* CachedMeshUid = ImplicitPrototypePrimPathToCachedMeshUid.Find(PrototypePrim.GetPrimPath()))
+				{
+					OutMeshUid = *CachedMeshUid;
+					return true;
+				}
+			}
+			return false;
+		}
+
+		void SetCachedSkeletalMeshUid(const UE::FUsdPrim& Prim, EMeshCategory Category, const FString& MeshUid)
+		{
+			if (IsCachableSkeletalMeshPrim(Prim, Category))
+			{
+				UE::FUsdPrim PrototypePrim = UsdUtils::GetPrototypePrim(Prim);
+				ImplicitPrototypePrimPathToCachedMeshUid.Add(PrototypePrim.GetPrimPath(), MeshUid);
+			}
+		}
+
+		FString CreateMeshNode(const UE::FSdfPath& Path, EMeshCategory Category, const TArray<FString>& AssemblyDependencies = {})
+		{
+			FTraversalInfo* Info = TraversalInfoByCachedPrimPath.Find(Path);
+			if (!ensure(Info))
+			{
+				return {};
+			}
+
+			const UE::FUsdPrim Prim = AssemblyRootPrim.GetStage().GetPrimAtPath(Path);
+			if (!ensure(Prim.IsA(TEXT("Mesh"))))
+			{
+				return {};
+			}
+
+			// If this is a repeated instanced skeletal mesh part, reuse the previously generated uid. This is
+			// a workaround for the current lack of instanced skeletal mesh support in Interchange.
+			if (FString CachedMeshUid; GetCachedSkeletalMeshUid(Prim, Category, CachedMeshUid))
+			{
+				return CachedMeshUid;
+			}
+
+			// If the mesh node we're creating has assembly dependencies, it must be the first base mesh, and so 
+            // attach the traversal result to our traversal info for retrieval during payload creation.
+			if (AssemblyDependencies.Num() > 0 && ensure(Result.IsValid()))
+			{
+				Info->NaniteAssemblyTraversalResult = Result;
+			}
+			
+			FString MeshUid = TranslatorImpl.AddMeshNode(Prim, TranslatorSettings, NodeContainer, *Info, false, AssemblyDependencies);
+
+			const FString SceneNodeUid = MakeNodeUid(Prim.GetPrimPath().GetString());
+			const FString DisplayLabel(Prim.GetName().ToString());
+			static const FString NoParentUid = TEXT("");
+			UInterchangeSceneNode* SceneNode = CreateSceneNode(Prim, TranslatorSettings, NodeContainer, SceneNodeUid, DisplayLabel, NoParentUid);
+
+			constexpr bool bResetCache = false;
+			SceneNode->SetCustomLocalTransform(&NodeContainer, FTransform::Identity, bResetCache);
+			SceneNode->SetCustomAssetInstanceUid(MeshUid);
+
+			// Store the mesh uid for the instanced skeletal mesh case
+			SetCachedSkeletalMeshUid(Prim, Category, MeshUid);
+
+			return MeshUid;
+		}
+
+		UE::FSdfPath GetSkelBindingPrimPathForPart(const UE::FUsdPrim& Prim, const FTraversalInfo& Info)
+		{
+			// This should only be called for skeletal mesh parts.
+			// 
+			// Check and get the path to the (possibly ancestor) prim holding skel binding data
+
+			UE::FSdfPath EmptyPath;
+			if (!MeshTypeIs(ENaniteAssemblyMeshType::SkeletalMesh))
+			{
+				return EmptyPath;
+			}
+
+			if (Info.ClosestNaniteAssemblySkelBindingPath)
+			{
+				const UE::FSdfPath SkelBindingPrimPath = UE::FSdfPath(**Info.ClosestNaniteAssemblySkelBindingPath);
+				if (SkelBindingPrimPath.IsPrimPath())
+				{
+					return SkelBindingPrimPath;
+				}
+			}
+
+			InvalidPartMeshNoSkelBindingPaths.Add(Prim.GetPrimPath());
+
+			return EmptyPath;
+		}
+
+		EMeshCategory GetMeshCategory(const UE::FUsdPrim& Prim, const FTraversalInfo& Info, TArray<FPrimPrototypeEntry>& PrototypeStack)
+		{
+			if (!Prim.IsA(TEXT("Mesh")))
+			{
+				return EMeshCategory::None;
+			}
+
+			// A mesh in a skeletal assembly must:
+			//  i: have a bound skeleton 
+			// ii: if it's a 'base' mesh (i.e. connected to the designated assembly root skeleton)
+			//      - must NOT be inside a pointinstancer
+			//      - must NOT be USD native instanced
+			//     or, if it's a 'part' mesh (i.e. attached to some other skeleton)
+			//      - may be inside a pointinstancer, but, not nested, or
+			//      - must be USD native instanced (because we don't want to encourage users to create tons of unique non-instanced parts) 
+			//
+			if (MeshTypeIs(ENaniteAssemblyMeshType::SkeletalMesh) && Info.BoundSkeletonPrimPath)
+			{
+				if (BaseSkeletonPrimPath.GetString() == **Info.BoundSkeletonPrimPath) // base skeleton mesh
+				{
+					if (PrototypeStack.IsEmpty() && !Prim.IsInstanceProxy())
+					{
+						return EMeshCategory::Base;
+					}
+				}
+				else // part skeleton mesh
+				{
+					if (PrototypeStack.Num() == 1 || Prim.IsInstanceProxy())
+					{
+						return EMeshCategory::Part;
+					}
+				}
+			}
+
+			// A part mesh in a static assembly must 
+			//  i: not have skel binding data
+			// ii: must be inside a pointinstancer prototype, or, be USD native instanced.
+			//
+			else if (MeshTypeIs(ENaniteAssemblyMeshType::StaticMesh))
+			{
+				if (!Prim.HasAPI(TEXT("SkelBindingAPI")) && (PrototypeStack.Num() > 0 || Prim.IsInstanceProxy()))
+				{
+					return EMeshCategory::Part;
+				}
+			}
+
+			// TODO: communicate all of the above to the user, somehow
+			return EMeshCategory::None;
+		}
+
+		void HandleOrIgnoreMesh(
+			const UE::FUsdPrim& Prim, 
+			FTraversalInfo Info, 
+			TArray<FTransform>& TransformStack, 
+			TArray<FPrimPrototypeEntry>& PrototypeStack)
+		{
+			if (!Prim.IsA(TEXT("Mesh")) || !Result.IsValid())
+			{
+				return;
+			}
+
+			EMeshCategory Category = GetMeshCategory(Prim, Info, PrototypeStack);
+			if (Category == EMeshCategory::None)
+			{
+				return;
+			}
+
+			// If this is a skeletal mesh part, check we have a prim with binding data
+			UE::FSdfPath SkelBindingPrimPath;
+			if (MeshTypeIs(ENaniteAssemblyMeshType::SkeletalMesh) && Category == EMeshCategory::Part)
+			{
+				SkelBindingPrimPath = GetSkelBindingPrimPathForPart(Prim, Info);
+				if (!SkelBindingPrimPath.IsPrimPath())
+				{
+					return;
+				}
+			}
+
+			ensure(
+				Result->AddEntry(MakeShared<FMeshEntry>(
+					Category, 
+					Prim.GetPrimPath(), 
+					SkelBindingPrimPath, 
+					TransformStack, 
+					PrototypeStack)));
+
+			TraversalInfoByCachedPrimPath.Add(Prim.GetPrimPath(), Info);
+		}
+
+		void HandleExternalAssetReference(
+			const UE::FUsdPrim& Prim,
+			FTraversalInfo Info,
+			TArray<FTransform>& TransformStack,
+			TArray<FPrimPrototypeEntry>& PrototypeStack,
+			const FString& AssetReference)
+		{
+			if (!Result.IsValid())
+			{
+				return;
+			}
+
+			UE::FSdfPath SkelBindingPrimPath;
+
+			// External asset references are only supported on parts.
+			EMeshCategory Category = EMeshCategory::Part;
+
+			if (MeshTypeIs(ENaniteAssemblyMeshType::SkeletalMesh))
+			{
+				// If we found the schema applied to a descendant of the base skelroot prim
+				// we should skip as we're not a part.
+				if (Prim.GetPrimPath().HasPrefix(BaseSkelRootPrimPath))
+				{
+					return;
+				}
+
+				// This is a skeletal mesh part - check we have prim with binding data
+				SkelBindingPrimPath = GetSkelBindingPrimPathForPart(Prim, Info);
+				if (!SkelBindingPrimPath.IsPrimPath())
+				{
+					return;
+				}
+			}
+
+			TSharedPtr<FMeshEntry> Entry = MakeShared<FMeshEntry>(
+				Category,
+				Prim.GetPrimPath(),
+				SkelBindingPrimPath,
+				TransformStack,
+				PrototypeStack);
+
+			if (!ensure(Result->AddEntry(Entry)))
+			{
+				return;
+			}
+			
+			Entry->NodeUid = AssetReference;
+
+			TraversalInfoByCachedPrimPath.Add(Prim.GetPrimPath(), Info);
+		}
+
+		void HandleSkeleton(const UE::FUsdPrim& Prim, FTraversalInfo Info, TArray<FTransform>& TransformStack) const
+		{
+			const bool bIsBaseSkeleton = Prim.GetPrimPath() == BaseSkeletonPrimPath;
+
+			const FString SceneNodeUid = MakeNodeUid(Prim.GetPrimPath().GetString());
+			const FString DisplayLabel(Prim.GetName().ToString());
+			const FString ThisParentSceneUid = bIsBaseSkeleton ? ParentSceneNodeUid : TEXT("");
+			UInterchangeSceneNode* SceneNode = CreateSceneNode(Prim, TranslatorSettings, NodeContainer, SceneNodeUid, DisplayLabel, ThisParentSceneUid);
+
+
+			const FTransform Transform = bIsBaseSkeleton ? TransformStack.Last() : FTransform::Identity;
+			constexpr bool bResetCache = false;
+			SceneNode->SetCustomLocalTransform(&NodeContainer, Transform, bResetCache);
+
+			TranslatorImpl.AddSkeletonNodes(Prim, *SceneNode, NodeContainer, Info);
+		}
+
+		bool FindAndSetBaseSkeleton()
+		{
+			if (!MeshTypeIs(ENaniteAssemblyMeshType::SkeletalMesh))
+			{
+				return true;
+			}
+
+			// Get the "base" skeleton prim as specified by the schemas "unreal:naniteAssembly:skeleton" relationship
+			
+			const UE::FUsdPrim SkelPrim = UsdToUnreal::NaniteAssemblyUtils::GetBaseSkeletonPrimForSkeletalMeshAssembly(AssemblyRootPrim);
+			if (!SkelPrim)
+			{
+				// Note reason for failure will have been logged by utility call above.
+				return false;
+			}
+			BaseSkeletonPrimPath = SkelPrim.GetPrimPath();
+
+			// We got the skeleton prim, so now check we have an ancestor SkelRoot prim before the assembly root prim in order to try
+			// detect weird or broken hierarchy configurations.
+
+			UE::FUsdPrim CurrentPrim = SkelPrim;
+			while (!CurrentPrim.IsPseudoRoot())
+			{
+				if (CurrentPrim.IsA(TEXT("SkelRoot")))
+				{
+					BaseSkelRootPrimPath = CurrentPrim.GetPrimPath();
+					break;
+				}
+				CurrentPrim = CurrentPrim.GetParent();
+			}
+
+			if (BaseSkelRootPrimPath.IsEmpty())
+			{
+				UE_LOG(LogUsd, Warning, 
+					TEXT("Failed to find SkelRoot prim for base Skeleton prim '%s' while processing NaniteAssemblyRoot prim '%s'")
+					, *BaseSkeletonPrimPath.GetString()
+					, *AssemblyRootPrim.GetPrimPath().GetString()
+				);
+				return false;
+			}
+			
+			return true;
+		}
+
+		bool ComputeTraversalResult(FTraversalInfo Info, const FTransform& ParentGlobalTransform)
+		{
+			Result = MakeShared<FNaniteAssemblyTraversalResult>(AssemblyMeshType, AssemblyRootPrim.GetPrimPath());
+
+			TArray<FTransform> TransformStack = { ParentGlobalTransform };
+			TArray<FPrimPrototypeEntry> PrototypeStack;
+			constexpr bool bTraverseInstanceProxies = true;
+			for (const UE::FUsdPrim& ChildPrim : AssemblyRootPrim.GetFilteredChildren(bTraverseInstanceProxies))
+			{
+				TraverseImpl(ChildPrim, Info, TransformStack, PrototypeStack);
+			}
+
+			if (!ValidateTraversalResult())
+			{
+				Result.Reset();
+				return false;
+			}
+
+			return true;
+		}
+
+		void TraversePointInstancerHierarchy(
+			const UE::FUsdPrim& Prim,
+			FTraversalInfo Info,
+			TArray<FTransform> TransformStack,
+			TArray<FPrimPrototypeEntry>& PrototypeStack
+		)
+		{
+			const TArray<UE::FSdfPath> PrototypePaths = UsdToUnreal::NaniteAssemblyUtils::GetPointInstancerPrototypePaths(Prim);
+			if (PrototypePaths.IsEmpty())
+			{
+				return;
+			}
+
+			for (int32 PrototypeIndex = 0; PrototypeIndex < PrototypePaths.Num(); ++PrototypeIndex)
+			{
+				UE::FUsdPrim PrototypePrim = Prim.GetStage().GetPrimAtPath(PrototypePaths[PrototypeIndex]);
+				if (!PrototypePrim)
+				{
+					continue;
+				}
+
+				FTraversalInfo PrototypePrimInfo = Info;
+				PrototypePrimInfo.UpdateWithCurrentPrim(PrototypePrim);
+
+				// Push transforms and prototype before entering each pointinstancer prototype
+				TransformStack.Add(FTransform::Identity);
+				PrototypeStack.Emplace(Prim.GetPrimPath(), PrototypePrim.GetPrimPath(), PrototypeIndex, PrototypePaths.Num());
+
+				TraverseImpl(PrototypePrim, Info, TransformStack, PrototypeStack);
+
+				// Pop before entering the next
+				PrototypeStack.Pop();
+				TransformStack.Pop();
+			}
+		}
+
+		void TraverseImpl(const UE::FUsdPrim& Prim, FTraversalInfo Info, TArray<FTransform> TransformStack, TArray<FPrimPrototypeEntry>& PrototypeStack)
+		{
+			if (!ensure(TransformStack.Num() > 0))
+			{
+				return;
+			}
+
+			if (Prim.HasAPI(TEXT("NaniteAssemblyRootAPI")))
+			{
+				UE_LOG(LogUsd, Warning,
+					TEXT("Skipping nested NaniteAssemblyRootAPI prim '%s' found below current root '%s'")
+					, *Prim.GetPrimPath().GetString()
+					, *AssemblyRootPrim.GetPrimPath().GetString()
+				);
+				return;
+			}
+
+			if (Prim.IsA(TEXT("Material")))
+			{
+				return;
+			}
+
+			Info.UpdateWithCurrentPrim(Prim);
+
+			if (Prim.IsA(TEXT("Xformable")))
+			{
+				FTransform LocalTransform = FTransform::Identity;
+				bool bResetTransformStack = false;
+				UsdToUnreal::ConvertXformable(Prim.GetStage(), UE::FUsdTyped(Prim), LocalTransform, UsdUtils::GetEarliestTimeCode(), &bResetTransformStack);
+				
+				FTransform& CurrentTransform = TransformStack.Last();
+				CurrentTransform = LocalTransform * CurrentTransform;
+			}
+
+			if (FString AssetRef; UsdToUnreal::NaniteAssemblyUtils::GetExternalAssetReference(Prim, AssetRef))
+			{
+				HandleExternalAssetReference(Prim, Info, TransformStack, PrototypeStack, AssetRef);
+				// The user is mostly likely expecting the subtree below the prim with the external
+				// reference applied to be pruned, so we stop the traversal.
+				return;
+			}
+
+			if (Prim.IsA(TEXT("Mesh")))
+			{
+				HandleOrIgnoreMesh(Prim, Info, TransformStack, PrototypeStack);
+			}
+			else if (Prim.IsA(TEXT("PointInstancer")))
+			{
+				TraversePointInstancerHierarchy(Prim, Info, TransformStack, PrototypeStack);
+			}
+			else if (Prim.IsA(TEXT("Skeleton")))
+			{
+				HandleSkeleton(Prim, Info, TransformStack);
+			}
+			else
+			{
+				constexpr bool bTraverseInstanceProxies = true;
+				for (const UE::FUsdPrim& ChildPrim : Prim.GetFilteredChildren(bTraverseInstanceProxies))
+				{
+					TraverseImpl(ChildPrim, Info, TransformStack, PrototypeStack);
+				}
+			}
+		}
+
+		bool ValidateHierarchy(const FTraversalInfo& Info) const
+		{
+			if (AssemblyMeshType == ENaniteAssemblyMeshType::None)
+			{
+				return false;
+			}
+
+			// Nanite assemblies cannot be inside LODs.
+			if (Info.bInsideLOD || Info.bIsLODContainer)
+			{
+				UE_LOG(LogUsd, Warning,
+					TEXT("Ignoring NaniteAssemblyRootAPI prim '%s' which is inside an LOD (not supported).")
+					, *AssemblyRootPrim.GetPrimPath().GetString()
+				);
+				return false;
+			}
+
+			// Check ancestors for configurations that we know will cause problems:
+			
+			UE::FUsdPrim CurrentPrim = AssemblyRootPrim;
+
+			while (!CurrentPrim.IsPseudoRoot())
+			{
+				if (CurrentPrim.IsA(TEXT("PointInstancer")))
+				{
+					UE_LOG(LogUsd, Warning,
+						TEXT("NaniteAssemblyRoot '%s' cannot currently be a descendant of a pointinstancer - '%s'")
+						, *AssemblyRootPrim.GetPrimPath().GetString()
+						, *CurrentPrim.GetPrimPath().GetString()
+					);
+					return false;
+				}
+				else if (CurrentPrim.HasAPI(*UsdToUnreal::ConvertToken(UnrealIdentifiers::LodSubtreeAPI)))
+				{
+					UE_LOG(LogUsd, Warning,
+						TEXT("NaniteAssemblyRoot '%s' cannot be nested inside an %s container.")
+						, *AssemblyRootPrim.GetPrimPath().GetString()
+						, *UsdToUnreal::ConvertToken(UnrealIdentifiers::LodSubtreeAPI)
+					);
+					return false;
+				}
+				CurrentPrim = CurrentPrim.GetParent();
+			}
+
+			return true;
+		}
+
+		bool ValidateTraversalResult() const
+		{
+			if (!Result)
+			{
+				return false;
+			}
+
+			// Report skel mesh parts that were dropped due to missing bindings
+			if (InvalidPartMeshNoSkelBindingPaths.Num() > 0)
+			{
+				FString Message = FString::Printf(
+					TEXT("Ignoring (%d) skeletal mesh prims missing joint binding data while generating parts for Nanite assembly ")
+					TEXT("prim '%s'.\nPlease check that the %s schema is applied to an ancestor prim and valid primvar data ")
+					TEXT("authored for attributes '%s' and '%s'...\n")
+					, InvalidPartMeshNoSkelBindingPaths.Num()
+					, *AssemblyRootPrim.GetPrimPath().GetString()
+					, *UsdToUnreal::ConvertToken(UnrealIdentifiers::NaniteAssemblySkelBindingAPI)
+					, *UsdToUnreal::ConvertToken(UnrealIdentifiers::PrimvarsUnrealNaniteAssemblyBindJoints)
+					, *UsdToUnreal::ConvertToken(UnrealIdentifiers::PrimvarsUnrealNaniteAssemblyBindJointWeights)
+				);
+				for (int32 Index = 0; Index < InvalidPartMeshNoSkelBindingPaths.Num(); ++Index)
+				{
+					// there might be a massive number of prims here, so only print the first few
+					int32 constexpr MaxEntries = 50;
+					if (Index < MaxEntries)
+					{
+						Message.Append(FString::Printf(TEXT("... %s\n"), *InvalidPartMeshNoSkelBindingPaths[Index].GetString()));
+					}
+					else
+					{
+						Message.Append(FString::Printf(TEXT("... (skipped %d more)"), InvalidPartMeshNoSkelBindingPaths.Num() - MaxEntries));
+						break;
+					}
+				}
+
+				UE_LOG(LogUsd, Warning, TEXT("%s"), *Message);
+			}
+
+			// Check we found some parts, and, if this is a skeletal mesh assembly, we also require at least one base mesh
+			const bool bHasBaseMeshes = Result->HasEntriesForCategory(EMeshCategory::Base);
+			const bool bHasPartMeshes = Result->HasEntriesForCategory(EMeshCategory::Part);
+			if (!bHasPartMeshes || (MeshTypeIs(ENaniteAssemblyMeshType::SkeletalMesh) && !bHasBaseMeshes))
+			{
+				UE_LOG(LogUsd, Warning,
+					TEXT("Failed to find any valid base/part meshes for NaniteAssemblyRootAPI prim '%s'")
+					, *AssemblyRootPrim.GetPrimPath().GetString()
+				);
+
+				return false;
+			}
+
+			return true;
+		}
+
+	public:
+
+		void Traverse(const UE::FUsdPrim& InAssemblyRootPrim, FTraversalInfo Info, const FTransform& ParentGlobalTransform)
+		{
+			AssemblyRootPrim = InAssemblyRootPrim;
+
+			if (!AssemblyRootPrim.HasAPI(TEXT("NaniteAssemblyRootAPI")))
+			{
+				return;
+			}
+
+			// Setup and initial validation
+
+			AssemblyMeshType = UsdToUnreal::NaniteAssemblyUtils::GetNaniteAssemblyMeshType(AssemblyRootPrim);
+
+			if (!ValidateHierarchy(Info) || !FindAndSetBaseSkeleton())
+			{
+				return;
+			}
+
+			// Compute the traversal result that will be attached to the traversal info object of
+			// the skeletal or static assembly base mesh. 
+			if (!ComputeTraversalResult(Info, ParentGlobalTransform) || !Result.IsValid())
+			{
+				return;
+			}
+			
+			// Now create mesh and scene nodes for parts. Do these first because we will need the resulting node 
+			// uids when creating the base meshes in a moment, in order to setup the correct dependencies.
+			for (TSharedPtr<FMeshEntry>& Entry : Result->GetEntriesForCategory(EMeshCategory::Part))
+			{
+				// Note that external asset ref part entries will have already set the uid to whatever asset.
+				if (Entry->NodeUid.IsEmpty()) 
+				{
+					// If create mesh fails here we'll wind up with an invalid/empty node uid for the entry.
+					// That's ok however because the payload processing will validate all entries anyway.
+					Entry->NodeUid = CreateMeshNode(Entry->PrimPath, Entry->Category);
+				}
+			}
+
+			// Skeletal mesh - now create base meshes and register dependencies on the parts we created
+			// above. Note also that we only need to apply the dependencies to one of the base meshes
+			// since Interchange will currently combine all the meshes anyway.
+			if (MeshTypeIs(ENaniteAssemblyMeshType::SkeletalMesh))
+			{
+				static const TArray<FString> EmptyDependencyArray;
+				bool bIsFirst = true;
+				for (TSharedPtr<FMeshEntry>& Entry : Result->GetEntriesForCategory(EMeshCategory::Base))
+				{
+					const TArray<FString> PartDependencies = bIsFirst 
+						? Result->GetMeshUidsForCategory(EMeshCategory::Part) 
+						: EmptyDependencyArray;
+
+					Entry->NodeUid = CreateMeshNode(Entry->PrimPath, Entry->Category, PartDependencies);
+
+					bIsFirst = false;
+				}
+			}
+			// Static mesh - the NaniteAssemblyRootAPI schema doesn't currently support specifying a base 
+			// mesh when in static mesh mode. So here we create a dummy mesh node that will result in an 
+			// empty mesh that the mesh factory can eventually add the assembly data to. 
+			else if (MeshTypeIs(ENaniteAssemblyMeshType::StaticMesh))
+			{
+				const FString MeshUid = MakeNodeUid(MeshPrefix + AssemblyRootPrim.GetPrimPath().GetString());
+				const FString NodeName(AssemblyRootPrim.GetPrimPath().GetName());
+
+				UInterchangeMeshNode* MeshNode = NewObject<UInterchangeMeshNode>(&NodeContainer);
+				NodeContainer.SetupNode(MeshNode, MeshUid, NodeName, EInterchangeNodeContainerType::TranslatedAsset);
+				MeshNode->SetAssetName(NodeName);
+				MeshNode->SetPayLoadKey(AssemblyRootPrim.GetPrimPath().GetString(), EInterchangeMeshPayLoadType::STATIC);
+
+				for (const FString& PartUid : Result->GetMeshUidsForCategory(EMeshCategory::Part))
+				{
+					MeshNode->AddAssemblyPartDependencyUid(PartUid);
+				}
+
+				ParentSceneNode->SetCustomAssetInstanceUid(MeshUid);
+
+				// Store the traversal result and cache the info for the mesh
+				Info.NaniteAssemblyTraversalResult = Result;
+				FWriteScopeLock ScopedInfoWriteLock{ TranslatorImpl.CachedTraversalInfoLock };
+				TranslatorImpl.NodeUidToCachedTraversalInfo.Add(MeshUid, Info);
+			}
+		}
+	};
+
+	void UInterchangeUSDTranslatorImpl::AddNaniteAssemblyHierarchy(
+		const UE::FUsdPrim& Prim,
+		UInterchangeUsdTranslatorSettings* TranslatorSettings,
+		UInterchangeBaseNodeContainer& NodeContainer,
+		UInterchangeSceneNode* ParentSceneNode,
+		FTraversalInfo& Info)
+	{
+		FNaniteAssemblyTraversalHelper NaniteAssemblyTraversalHelper(*this, NodeContainer, TranslatorSettings, ParentSceneNode);
+		NaniteAssemblyTraversalHelper.Traverse(Prim, Info, FTransform());
+	}
+
+	//Helper to traverse PointInstancer subhierarchy and process it into InstancedStaticMeshComponents (only).
+	struct FISMComponentTraversalHelper
+	{
+		UInterchangeUSDTranslatorImpl& TranslatorImpl;
+		UInterchangeBaseNodeContainer& NodeContainer;
+		UInterchangeUsdTranslatorSettings* TranslatorSettings;
+		UInterchangeSceneNode* ParentSceneNode;
+		const FString ParentSceneNodeUid;
+
+		//While we could lookup from the NodeContainer, 
+		// we are implementing a LookUp table specifically for ISMComponentNodes for speed:
+		TMap<FString, UInterchangeInstancedStaticMeshComponentNode*> UidToComponentNode;
+
+		FISMComponentTraversalHelper(
+			UInterchangeUSDTranslatorImpl& InTranslatorImpl,
+			UInterchangeBaseNodeContainer& InNodeContainer,
+			UInterchangeUsdTranslatorSettings* InTranslatorSettings,
+			UInterchangeSceneNode* InParentSceneNode)
+			: TranslatorImpl(InTranslatorImpl)
+			, NodeContainer(InNodeContainer)
+			, TranslatorSettings(InTranslatorSettings)
+			, ParentSceneNode(InParentSceneNode)
+			, ParentSceneNodeUid(ParentSceneNode->GetUniqueID())
+		{
+
+		}
+
 		UInterchangeInstancedStaticMeshComponentNode* CreateSceneComponent(const FString& MeshUid)
 		{
 			FString UniqueId = MakeNodeUid(ComponentPrefix + TEXT("\\") + ParentSceneNodeUid + TEXT("\\") + MeshUid);
@@ -4748,324 +5102,4 @@
 #endif	  // WITH_EDITOR
 	}
 
-	template<typename USDType, typename InterchangeType = USDType>
-	InterchangeType CoerceToInterchangeType(const USDType& USDValue)
-	{
-		if constexpr (std::is_same_v<USDType, FMatrix2D> && std::is_same_v<InterchangeType, FMatrix44d>)
-		{
-			FMatrix44d Result{EForceInit::ForceInitToZero};
-			Result.M[0][0] = USDValue.Row0.X;
-			Result.M[0][1] = USDValue.Row0.Y;
-
-			Result.M[1][0] = USDValue.Row1.X;
-			Result.M[1][1] = USDValue.Row1.Y;
-			return Result;
-		}
-		else if constexpr (std::is_same_v<USDType, FMatrix3D> && std::is_same_v<InterchangeType, FMatrix44d>)
-		{
-			FMatrix44d Result{EForceInit::ForceInitToZero};
-			Result.M[0][0] = USDValue.Row0.X;
-			Result.M[0][1] = USDValue.Row0.Y;
-			Result.M[0][2] = USDValue.Row0.Z;
-
-			Result.M[1][0] = USDValue.Row1.X;
-			Result.M[1][1] = USDValue.Row1.Y;
-			Result.M[1][2] = USDValue.Row1.Z;
-
-			Result.M[2][0] = USDValue.Row2.X;
-			Result.M[2][1] = USDValue.Row2.Y;
-			Result.M[2][2] = USDValue.Row2.Z;
-			return Result;
-		}
-		else
-		{
-			return (InterchangeType)USDValue;
-		}
-	};
-
-	template<typename USDType, typename InterchangeType = USDType>
-	void AddUserAttribute(const FString& KeyName, const UE::FVtValue& VtValue, UInterchangeSceneNode* Node)
-	{
-		const USDType USDValue = VtValue.Get<USDType>();
-		const TOptional<FString> PayloadKey;
-
-		if constexpr (std::is_same_v<InterchangeType, USDType>)
-		{
-			UInterchangeUserDefinedAttributesAPI::CreateUserDefinedAttribute(Node, KeyName, USDValue, PayloadKey);
-		}
-		else if constexpr (TIsTArray<USDType>::Value)
-		{
-			static_assert(TIsTArray<InterchangeType>::Value);
-
-			InterchangeType InterchangeArray;
-			InterchangeArray.SetNum(USDValue.Num());
-			for (int32 Index = 0; Index < USDValue.Num(); ++Index)
-			{
-				InterchangeArray[Index] = CoerceToInterchangeType<typename USDType::ElementType, typename InterchangeType::ElementType>(USDValue[Index]);
-			}
-
-			UInterchangeUserDefinedAttributesAPI::CreateUserDefinedAttribute(Node, KeyName, InterchangeArray, PayloadKey);
-		}
-		else
-		{
-			const InterchangeType InterchangeValue = CoerceToInterchangeType<USDType, InterchangeType>(USDValue);
-			UInterchangeUserDefinedAttributesAPI::CreateUserDefinedAttribute(Node, KeyName, InterchangeValue, PayloadKey);
-		}
-	};
-
-	// Note: We could have just used the Cpp type names here instead of the value type token, but doing the
-	// latter lets us handle the color types as FLinearColor, which is probably more useful
-	// References:
-	// - https://openusd.org/docs/api/_usd__page__datatypes.html
-	// - EAttributeTypes declaration
-	// - Engine/Source/Runtime/Interchange/Core/Private/Tests/StorageTest.cpp
-	using AddUserAttributeFunc = TFunction<void(const FString&, const UE::FVtValue&, UInterchangeSceneNode*)>;
-	static TMap<FName, AddUserAttributeFunc> AddUserAttributeFuncs = {
-		{TEXT("bool"),  					AddUserAttribute<bool>},
-
-		{TEXT("uchar"), 					AddUserAttribute<uint8>},
-		{TEXT("unsigned char"), 			AddUserAttribute<uint8>},
-		{TEXT("int"), 						AddUserAttribute<int32>},
-		{TEXT("unsigned int"), 				AddUserAttribute<uint32>},
-		{TEXT("uint"), 						AddUserAttribute<uint32>},
-		{TEXT("int64"), 					AddUserAttribute<int64>},
-		{TEXT("__int64"), 					AddUserAttribute<int64>},
-		{TEXT("unsigned __int64"), 			AddUserAttribute<uint64>},
-		{TEXT("uint64"), 					AddUserAttribute<uint64>},
-
-		{TEXT("half"),	 					AddUserAttribute<FFloat16>},
-		{TEXT("pxr_half::half"),	 		AddUserAttribute<FFloat16>},
-		{TEXT("float"), 					AddUserAttribute<float>},
-		{TEXT("double"), 					AddUserAttribute<double>},
-		{TEXT("timecode"), 					AddUserAttribute<FSdfTimeCode, double>},
-		{TEXT("SdfTimeCode"), 				AddUserAttribute<FSdfTimeCode, double>},
-
-		{TEXT("string"), 					AddUserAttribute<FString>},
-		{TEXT("token"), 					AddUserAttribute<FName>},
-		{TEXT("TfToken"), 					AddUserAttribute<FName>},
-		{TEXT("asset"), 					AddUserAttribute<FSdfAssetPath, FString>},
-		{TEXT("SdfAssetPath"), 				AddUserAttribute<FSdfAssetPath, FString>},
-
-		{TEXT("matrix2d"), 					AddUserAttribute<FMatrix2D, FMatrix44d>},
-		{TEXT("GfMatrix2d"), 				AddUserAttribute<FMatrix2D, FMatrix44d>},
-		{TEXT("matrix3d"), 					AddUserAttribute<FMatrix3D, FMatrix44d>},
-		{TEXT("GfMatrix3d"), 				AddUserAttribute<FMatrix3D, FMatrix44d>},
-		{TEXT("matrix4d"), 					AddUserAttribute<FMatrix44d>},
-		{TEXT("GfMatrix4d"), 				AddUserAttribute<FMatrix44d>},
-		{TEXT("frame4d"), 					AddUserAttribute<FMatrix44d>},
-
-		{TEXT("quath"), 					AddUserAttribute<FQuat4h, FQuat4f>},
-		{TEXT("GfQuath"), 					AddUserAttribute<FQuat4h, FQuat4f>},
-		{TEXT("quatf"), 					AddUserAttribute<FQuat4f>},
-		{TEXT("GfQuatf"), 					AddUserAttribute<FQuat4f>},
-		{TEXT("quatd"), 					AddUserAttribute<FQuat4d>},
-		{TEXT("GfQuatd"), 					AddUserAttribute<FQuat4d>},
-
-		{TEXT("half2"), 					AddUserAttribute<FVector2DHalf>},
-		{TEXT("GfVec2h"), 					AddUserAttribute<FVector2DHalf>},
-		{TEXT("float2"), 					AddUserAttribute<FVector2f>},
-		{TEXT("GfVec2f"), 					AddUserAttribute<FVector2f>},
-		{TEXT("double2"), 					AddUserAttribute<FVector2d>},
-		{TEXT("GfVec2d"), 					AddUserAttribute<FVector2d>},
-		{TEXT("int2"), 						AddUserAttribute<FIntPoint>},
-		{TEXT("GfVec2i"), 					AddUserAttribute<FIntPoint>},
-
-		{TEXT("half3"), 					AddUserAttribute<FVector3h, FVector3f>},
-		{TEXT("GfVec3h"), 					AddUserAttribute<FVector3h, FVector3f>},
-		{TEXT("point3h"), 					AddUserAttribute<FVector3h, FVector3f>},
-		{TEXT("normal3h"), 					AddUserAttribute<FVector3h, FVector3f>},
-		{TEXT("vector3h"), 					AddUserAttribute<FVector3h, FVector3f>},
-		{TEXT("color3h"), 					AddUserAttribute<FVector3h, FLinearColor>},
-
-		{TEXT("float3"), 					AddUserAttribute<FVector3f>},
-		{TEXT("GfVec3f"), 					AddUserAttribute<FVector3f>},
-		{TEXT("point3f"), 					AddUserAttribute<FVector3f>},
-		{TEXT("normal3f"), 					AddUserAttribute<FVector3f>},
-		{TEXT("vector3f"), 					AddUserAttribute<FVector3f>},
-		{TEXT("color3f"), 					AddUserAttribute<FVector3f, FLinearColor>},
-
-		{TEXT("double3"), 					AddUserAttribute<FVector3d>},
-		{TEXT("GfVec3d"), 					AddUserAttribute<FVector3d>},
-		{TEXT("point3d"), 					AddUserAttribute<FVector3d>},
-		{TEXT("normal3d"), 					AddUserAttribute<FVector3d>},
-		{TEXT("vector3d"), 					AddUserAttribute<FVector3d>},
-		{TEXT("color3d"), 					AddUserAttribute<FVector3d, FLinearColor>},
-
-		{TEXT("int3"), 						AddUserAttribute<FIntVector>},
-		{TEXT("GfVec3i"), 					AddUserAttribute<FIntVector>},
-
-		{TEXT("half4"), 					AddUserAttribute<FVector4h, FVector4f>},
-		{TEXT("GfVec4h"), 					AddUserAttribute<FVector4h, FVector4f>},
-		{TEXT("float4"), 					AddUserAttribute<FVector4f>},
-		{TEXT("GfVec4f"), 					AddUserAttribute<FVector4f>},
-		{TEXT("double4"), 					AddUserAttribute<FVector4d>},
-		{TEXT("GfVec4d"), 					AddUserAttribute<FVector4d>},
-		{TEXT("int4"), 						AddUserAttribute<FIntRect>},
-		{TEXT("GfVec4i"), 					AddUserAttribute<FIntRect>},
-		{TEXT("color4h"), 					AddUserAttribute<FVector4h, FLinearColor>},
-		{TEXT("color4f"), 					AddUserAttribute<FVector4f, FLinearColor>},
-		{TEXT("color4d"), 					AddUserAttribute<FVector4d, FLinearColor>},
-
-		{TEXT("bool[]"),  					AddUserAttribute<TArray<bool>>},
-		{TEXT("VtArray<bool>"),  			AddUserAttribute<TArray<bool>>},
-
-		{TEXT("uchar[]"), 					AddUserAttribute<TArray<uint8>>},
-		{TEXT("VtArray<unsigned char>"), 	AddUserAttribute<TArray<uint8>>},
-		{TEXT("int[]"), 					AddUserAttribute<TArray<int32>>},
-		{TEXT("VtArray<int>"), 				AddUserAttribute<TArray<int32>>},
-		{TEXT("uint[]"), 					AddUserAttribute<TArray<uint32>>},
-		{TEXT("VtArray<unsigned int>"), 	AddUserAttribute<TArray<uint32>>},
-		{TEXT("int64[]"), 					AddUserAttribute<TArray<int64>>},
-		{TEXT("VtArray<__int64>"), 			AddUserAttribute<TArray<int64>>},
-		{TEXT("uint64[]"), 					AddUserAttribute<TArray<uint64>>},
-		{TEXT("VtArray<unsigned __int64>"), AddUserAttribute<TArray<uint64>>},
-
-		{TEXT("half[]"),	 				AddUserAttribute<TArray<FFloat16>>},
-		{TEXT("VtArray<pxr_half::half>"),	AddUserAttribute<TArray<FFloat16>>},
-		{TEXT("float[]"), 					AddUserAttribute<TArray<float>>},
-		{TEXT("VtArray<float>"), 			AddUserAttribute<TArray<float>>},
-		{TEXT("double[]"), 					AddUserAttribute<TArray<double>>},
-		{TEXT("VtArray<double>"), 			AddUserAttribute<TArray<double>>},
-		{TEXT("timecode[]"), 				AddUserAttribute<TArray<FSdfTimeCode>, TArray<double>>},
-		{TEXT("VtArray<SdfTimeCode>"), 		AddUserAttribute<TArray<FSdfTimeCode>, TArray<double>>},
-
-		{TEXT("string[]"), 					AddUserAttribute<TArray<FString>>},
-		{TEXT("VtArray<string>"), 			AddUserAttribute<TArray<FString>>},
-		{TEXT("token[]"), 					AddUserAttribute<TArray<FName>>},
-		{TEXT("VtArray<TfToken>"), 			AddUserAttribute<TArray<FName>>},
-		{TEXT("asset[]"), 					AddUserAttribute<TArray<FSdfAssetPath>, TArray<FString>>},
-		{TEXT("VtArray<SdfAssetPath>"), 	AddUserAttribute<TArray<FSdfAssetPath>, TArray<FString>>},
-
-		{TEXT("matrix2d[]"), 				AddUserAttribute<TArray<FMatrix2D>, TArray<FMatrix44d>>},
-		{TEXT("VtArray<GfMatrix2d>"), 		AddUserAttribute<TArray<FMatrix2D>, TArray<FMatrix44d>>},
-		{TEXT("matrix3d[]"), 				AddUserAttribute<TArray<FMatrix3D>, TArray<FMatrix44d>>},
-		{TEXT("VtArray<GfMatrix3d>"), 		AddUserAttribute<TArray<FMatrix3D>, TArray<FMatrix44d>>},
-		{TEXT("matrix4d[]"), 				AddUserAttribute<TArray<FMatrix44d>>},
-		{TEXT("VtArray<GfMatrix4d>"), 		AddUserAttribute<TArray<FMatrix44d>>},
-		{TEXT("frame4d[]"), 				AddUserAttribute<TArray<FMatrix44d>>},
-
-		{TEXT("quath[]"), 					AddUserAttribute<TArray<FQuat4h>, TArray<FQuat4f>>},
-		{TEXT("VtArray<GfQuath>"), 			AddUserAttribute<TArray<FQuat4h>, TArray<FQuat4f>>},
-		{TEXT("quatf[]"), 					AddUserAttribute<TArray<FQuat4f>>},
-		{TEXT("VtArray<GfQuatf>"), 			AddUserAttribute<TArray<FQuat4f>>},
-		{TEXT("quatd[]"), 					AddUserAttribute<TArray<FQuat4d>>},
-		{TEXT("VtArray<GfQuatd>"), 			AddUserAttribute<TArray<FQuat4d>>},
-
-		{TEXT("half2[]"), 					AddUserAttribute<TArray<FVector2DHalf>>},
-		{TEXT("VtArray<GfVec2h>"), 			AddUserAttribute<TArray<FVector2DHalf>>},
-		{TEXT("float2[]"), 					AddUserAttribute<TArray<FVector2f>>},
-		{TEXT("VtArray<GfVec2f>"), 			AddUserAttribute<TArray<FVector2f>>},
-		{TEXT("double2[]"), 				AddUserAttribute<TArray<FVector2d>>},
-		{TEXT("VtArray<GfVec2d>"), 			AddUserAttribute<TArray<FVector2d>>},
-		{TEXT("int2[]"), 					AddUserAttribute<TArray<FIntPoint>>},
-		{TEXT("VtArray<GfVec2i>"), 			AddUserAttribute<TArray<FIntPoint>>},
-
-		{TEXT("half3[]"), 					AddUserAttribute<TArray<FVector3h>, TArray<FVector3f>>},
-		{TEXT("VtArray<GfVec3h>"), 			AddUserAttribute<TArray<FVector3h>, TArray<FVector3f>>},
-		{TEXT("point3h[]"), 				AddUserAttribute<TArray<FVector3h>, TArray<FVector3f>>},
-		{TEXT("normal3h[]"), 				AddUserAttribute<TArray<FVector3h>, TArray<FVector3f>>},
-		{TEXT("vector3h[]"), 				AddUserAttribute<TArray<FVector3h>, TArray<FVector3f>>},
-		{TEXT("color3h[]"), 				AddUserAttribute<TArray<FVector3h>, TArray<FLinearColor>>},
-
-		{TEXT("float3[]"), 					AddUserAttribute<TArray<FVector3f>>},
-		{TEXT("VtArray<GfVec3f>"), 			AddUserAttribute<TArray<FVector3f>>},
-		{TEXT("point3f[]"), 				AddUserAttribute<TArray<FVector3f>>},
-		{TEXT("normal3f[]"), 				AddUserAttribute<TArray<FVector3f>>},
-		{TEXT("vector3f[]"), 				AddUserAttribute<TArray<FVector3f>>},
-		{TEXT("color3f[]"), 				AddUserAttribute<TArray<FVector3f>, TArray<FLinearColor>>},
-
-		{TEXT("double3[]"), 				AddUserAttribute<TArray<FVector3d>>},
-		{TEXT("VtArray<GfVec3d>"), 			AddUserAttribute<TArray<FVector3d>>},
-		{TEXT("point3d[]"), 				AddUserAttribute<TArray<FVector3d>>},
-		{TEXT("normal3d[]"), 				AddUserAttribute<TArray<FVector3d>>},
-		{TEXT("vector3d[]"), 				AddUserAttribute<TArray<FVector3d>>},
-		{TEXT("color3d[]"), 				AddUserAttribute<TArray<FVector3d>, TArray<FLinearColor>>},
-
-		{TEXT("int3[]"), 					AddUserAttribute<TArray<FIntVector>>},
-		{TEXT("VtArray<GfVec3i>"), 			AddUserAttribute<TArray<FIntVector>>},
-
-		{TEXT("half4[]"), 					AddUserAttribute<TArray<FVector4h>, TArray<FVector4f>>},
-		{TEXT("VtArray<GfVec4h>"), 			AddUserAttribute<TArray<FVector4h>, TArray<FVector4f>>},
-		{TEXT("float4[]"), 					AddUserAttribute<TArray<FVector4f>>},
-		{TEXT("VtArray<GfVec4f>"), 			AddUserAttribute<TArray<FVector4f>>},
-		{TEXT("double4[]"), 				AddUserAttribute<TArray<FVector4d>>},
-		{TEXT("VtArray<GfVec4d>"), 			AddUserAttribute<TArray<FVector4d>>},
-		{TEXT("int4[]"), 					AddUserAttribute<TArray<FIntRect>>},
-		{TEXT("VtArray<GfVec4i>"), 			AddUserAttribute<TArray<FIntRect>>},
-		{TEXT("color4h[]"), 				AddUserAttribute<TArray<FVector4h>, TArray<FLinearColor>>},
-		{TEXT("color4f[]"), 				AddUserAttribute<TArray<FVector4f>, TArray<FLinearColor>>},
-		{TEXT("color4d[]"), 				AddUserAttribute<TArray<FVector4d>, TArray<FLinearColor>>},
-	};
-
-	void TranslateAttributes(const UE::FUsdPrim& Prim, UInterchangeSceneNode* Node, const FString& AllowedAttributeRegex)
-	{
-		TRACE_CPUPROFILER_EVENT_SCOPE(TranslateAttributes)
-
-		FRegexPattern RegexPattern{AllowedAttributeRegex};
-
-		for (const UE::FUsdAttribute& Attr : Prim.GetAttributes())
-		{
-			if (!Attr.HasAuthoredValue())
-			{
-				continue;
-			}
-
-			const FString AttrName = Attr.GetName().ToString();
-
-			FRegexMatcher RegexMatcher{RegexPattern, AttrName};
-			if (!RegexMatcher.FindNext())
-			{
-				continue;
-			}
-
-			if (AddUserAttributeFunc Func = AddUserAttributeFuncs.FindRef(Attr.GetTypeName()))
-			{
-				UE::FVtValue Value;
-				if (Attr.Get(Value)) // Always check for an opinion on the default time code
-				{
-					Func(AttrName, Value, Node);
-				}
-			}
-		}
-	}
-
-	void TranslateMetadata(const UE::FUsdPrim& Prim, UInterchangeSceneNode* Node, const FString& AllowedMetadataRegex)
-	{
-		TRACE_CPUPROFILER_EVENT_SCOPE(TranslateMetadata)
-
-		FRegexPattern RegexPattern{AllowedMetadataRegex};
-
-		TFunction<void(const TMap<FString, UE::FVtValue>&, const FString&)> ParseMetadataMap;
-		ParseMetadataMap = [&ParseMetadataMap, &RegexPattern, Node](const TMap<FString, UE::FVtValue>& MetadataMap, const FString& ConcatParentKey)
-		{
-			for (const TPair<FString, UE::FVtValue>& Pair : MetadataMap)
-			{
-				const FString FullKey = ConcatParentKey.IsEmpty() ? Pair.Key : FString::Printf(TEXT("%s:%s"), *ConcatParentKey, *Pair.Key);
-				const FString TypeName = Pair.Value.GetTypeName();
-
-				if (TypeName == TEXT("VtDictionary"))
-				{
-					TMap<FString, UE::FVtValue> NestedMap = Pair.Value.Get<TMap<FString, UE::FVtValue>>();
-					ParseMetadataMap(NestedMap, FullKey);
-				}
-				else if (AddUserAttributeFunc Func = AddUserAttributeFuncs.FindRef(*TypeName))
-				{
-					// Only match when parsing leaf keys, because a key may not match just the parent nested map key directly and
-					// would have otherwise prevented us from stepping into it (e.g. the regex "customData:my" will not match 
-					// "customData", but it would match "customData:myNestedMap")
-					FRegexMatcher RegexMatcher{RegexPattern, FullKey};
-					if (!RegexMatcher.FindNext())
-					{
-						continue;
-					}
-
-					Func(FullKey, Pair.Value, Node);
-				}
-			}
-		};
-
-		const FString ConcatParentKey;
-		ParseMetadataMap(Prim.GetAllAuthoredMetadata(), ConcatParentKey);
-	}
-
 	void Traverse(
@@ -5168,34 +5202,7 @@
 		UInterchangeSceneNode* SceneNode = nullptr;
 		if (bNeedsSceneNode)
 		{
-			SceneNode = NewObject<UInterchangeSceneNode>(&NodeContainer);
-			NodeContainer.SetupNode(SceneNode, SceneNodeUid, DisplayLabel, EInterchangeNodeContainerType::TranslatedScene, Info.ParentNode ? Info.ParentNode->GetUniqueID() : TEXT(""));
+			SceneNode = CreateSceneNode(Prim, TranslatorSettings, NodeContainer, SceneNodeUid, DisplayLabel, Info.ParentNode ? Info.ParentNode->GetUniqueID() : TEXT(""));
 
-			// Store our purpose here to be filtered on the pipeline
-			const EUsdPurpose PrimPurpose = IUsdPrim::GetPurpose(Prim);
-			if (PrimPurpose != EUsdPurpose::Default)
-			{
-				SceneNode->AddInt32Attribute(UE::Interchange::USD::GeometryPurposeIdentifier, (int32)PrimPurpose);
-			}
-
-			// Store our prim kind as well, if we have any (this becomes the empty string if the prim has no authored kind)
-			const FString KindString = UsdToUnreal::ConvertToken(IUsdPrim::GetKind(Prim));
-			if (!KindString.IsEmpty())
-			{
-				const FString UserDefinedAttributeName = TEXT("kind");
-				TOptional<FString> PayloadKey;
-				UInterchangeUserDefinedAttributesAPI::CreateUserDefinedAttribute(SceneNode, UserDefinedAttributeName, KindString, PayloadKey);
-			}
-
-			if (TranslatorSettings->bTranslatePrimAttributes)
-			{
-				TranslateAttributes(Prim, SceneNode, TranslatorSettings->AttributeRegexFilter);
-			}
-
-			if (TranslatorSettings->bTranslatePrimMetadata)
-			{
-				TranslateMetadata(Prim, SceneNode, TranslatorSettings->MetadataRegexFilter);
-			}
-
 			// If we're an Xformable, get our transform.
 			// All SceneNodes should have their LocalTransform set though.
