USD: Fix handling of nested LOD skeletal meshes

diff --git a/Engine/Plugins/Interchange/Extensions/OpenUSD/Source/Import/Private/InterchangeUsdTranslator.cpp b/Engine/Plugins/Interchange/Extensions/OpenUSD/Source/Import/Private/InterchangeUsdTranslator.cpp
--- a/Engine/Plugins/Interchange/Extensions/OpenUSD/Source/Import/Private/InterchangeUsdTranslator.cpp
+++ b/Engine/Plugins/Interchange/Extensions/OpenUSD/Source/Import/Private/InterchangeUsdTranslator.cpp
@@ -173,8 +173,10 @@ namespace UE::InterchangeUsdTranslator::Private
 
 			bVisible = bVisible && UsdUtils::HasInheritedVisibility(CurrentPrim);
 
-			// Check this first so that we go bInsideLOD if our parent was a LODContainer
-			bInsideLOD = bInsideLOD || bIsLODContainer;
+			// Check this first so that we go bInsideLOD if our parent was a LODContainer.
+			// bInsideLOD is only set for the immediate children of the prim that sets bIsLODContainer,
+			// as we only handle LOD variant meshes that are direct children of the LOD container prim anyway
+			bInsideLOD = bIsLODContainer;
 
 			// We only want this to be true when we're traversing the exact prim that owns
 			// the LOD: Once we step into any of its children it should go back to false
@@ -441,7 +443,7 @@ namespace UE::InterchangeUsdTranslator::Private
 		TMap<FString, UsdUtils::FUsdPrimMaterialAssignmentInfo> CachedMaterialAssignments;
 
 		// We fill this in while we're translating a LOD container so that we can do some post-processing inside FinalizeLODContainerTraversal
-		TArray<UInterchangeSceneNode*> CurrentLODSceneNodes;
+		TMap<UInterchangeSceneNode*, TArray<UInterchangeSceneNode*>> LODVariantSetNodeToLodMeshNodes;
 
 		// When traversing we'll generate FTraversalInfo objects. If we need to (e.g. for skinned meshes),
 		// we'll store the info for that translated node here, so we don't have to recompute it when returning
@@ -2670,7 +2672,7 @@ namespace UE::InterchangeUsdTranslator::Private
 				}
 				else
 				{
-					UE_LOG(LogTemp, Warning, TEXT("Failed to parse a LOD Mesh from variant '%s' of prim '%s'. For automatic parsing of LODs, make sure there is a single Mesh prim within the variant, named exactly as the variant itself (e.g. 'LOD0', 'LOD1', etc.)"), *VariantName, *Prim.GetPrimPath().GetString());
+					UE_LOG(LogUsd, Warning, TEXT("Failed to parse a LOD Mesh from variant '%s' of prim '%s'. For automatic parsing of LODs, make sure there is a single Mesh prim within the variant, named exactly as the variant itself (e.g. 'LOD0', 'LOD1', etc.)"), *VariantName, *Prim.GetPrimPath().GetString());
 				}
 			}
 
@@ -3091,16 +3093,15 @@ namespace UE::InterchangeUsdTranslator::Private
 				}
 			}
 
-			// If not, let's open a new masked stage.
-			// USD will only compose the LOD container prim (and its subtree), so this should be very small and quick to open.
-			// We won't use the stage cache for these, so our strong reference right here is the only thing holding the stage opened.
+			// Open a brand new stage we can freely flip our variants in without invalidating other prim references from concurrent tasks.
+			// We won't use the stage cache for these, so our strong reference right here is the only thing holding the stage opened
 			if (!TempStage)
 			{
-				TempStage = UnrealUSDWrapper::OpenMaskedStage(
-					*UsdStage.GetRootLayer().GetIdentifier(),
-					EUsdInitialLoadSet::LoadAll,
-					{LODContainerPathString}
-				);
+				// We used to use a population mask here to only compose the prim subtree of the LOD prim we want to parse, but
+				// that is too prone to edge case issues as any relationship to prims outside the loaded subtree would fail and 
+				// emit errors (blend shapes, material assignments, etc.)
+				const bool bUseStageCache = false;
+				TempStage = UnrealUSDWrapper::OpenStage(*UsdStage.GetRootLayer().GetIdentifier(), EUsdInitialLoadSet::LoadAll, bUseStageCache);
 
 				if (!TempStage)
 				{
@@ -3164,9 +3165,17 @@ namespace UE::InterchangeUsdTranslator::Private
 			int32 LODIndex;	   // Using actual ints in here instead of sorting on the node DisplayLabel directly so we correctly sort LOD10 > LOD2
 		};
 
+		// If we have less than two LODs, then we don't need to spawn a multi-LOD mesh and can just early out
+		TArray<UInterchangeSceneNode*>* CurrentLODSceneNodes = LODVariantSetNodeToLodMeshNodes.Find(SceneNodeWithLODs);
+		if (!CurrentLODSceneNodes || CurrentLODSceneNodes->Num() < 2)
+		{
+			LODVariantSetNodeToLodMeshNodes.Remove(SceneNodeWithLODs);
+			return;
+		}
+
 		TArray<FLODSortHelper> SortedLODs;
-		SortedLODs.Reserve(CurrentLODSceneNodes.Num());
-		for (UInterchangeSceneNode* Node : CurrentLODSceneNodes)
+		SortedLODs.Reserve(CurrentLODSceneNodes->Num());
+		for (UInterchangeSceneNode* Node : *CurrentLODSceneNodes)
 		{
 			SortedLODs.Add({Node, GetLODIndexFromName(Node->GetDisplayLabel())});
 		}
@@ -3180,7 +3189,7 @@ namespace UE::InterchangeUsdTranslator::Private
 		// Parent our LOD nodes to the LODContainer in the right order, because Interchange will try assigning
 		// those meshes to LOD numbers in the order it traverses the children, and there's no guarantee about the order we run
 		// into the variants when traversing the stage (or the order with which they were authored)
-		for (int32 Index = 0; Index < CurrentLODSceneNodes.Num(); ++Index)
+		for (int32 Index = 0; Index < CurrentLODSceneNodes->Num(); ++Index)
 		{
 			const FLODSortHelper& PackedNode = SortedLODs[Index];
 			const FString& NodeUID = PackedNode.SceneNode->GetUniqueID();
@@ -3245,7 +3254,7 @@ namespace UE::InterchangeUsdTranslator::Private
 			bIsFirstLOD = false;
 		}
 
-		CurrentLODSceneNodes.Reset();
+		LODVariantSetNodeToLodMeshNodes.Remove(SceneNodeWithLODs);
 	}
 
 	template<typename MeshOrSceneNodeType>
@@ -5188,9 +5197,18 @@ namespace UE::InterchangeUsdTranslator::Private
 				);
 			}
 		}
-		else
+
+		if (Prim.IsA(TEXT("Gprim")) && Prim.GetParent().IsA(TEXT("Gprim")))
 		{
-			// Handle the KnownUnsupportedTypes here.
+			UInterchangeResultWarning_Generic* Message = TranslatorImpl.ResultsContainer->Add<UInterchangeResultWarning_Generic>();
+			Message->SourceAssetName = TranslatorImpl.UsdStage.GetRootLayer().GetRealPath();
+			Message->Text = FText::Format(
+				LOCTEXT(
+					"UnsupportedSchemaType",
+					"Prim '{0}' is a Gprim child of another Gprim. This configuration is not recommended or supported very well, and may lead to unexpected results."
+				),
+				FText::FromString(Prim.GetPrimPath().GetString())
+			);
 		}
 
 		// Only prims that require rendering (and have a renderable parent) get a scene node.
@@ -5245,9 +5263,12 @@ namespace UE::InterchangeUsdTranslator::Private
 				TranslatorImpl.AddLODMeshNodes(Prim, NodeContainer, *SceneNode, TranslatorSettings, Info);
 			}
 
-			if (Info.bInsideLOD && Prim.IsA(TEXT("Mesh")))
+			if (Info.bInsideLOD && Prim.IsA(TEXT("Mesh")) && GetLODIndexFromName(SceneNode->GetDisplayLabel()) != INDEX_NONE)
 			{
-				TranslatorImpl.CurrentLODSceneNodes.Add(SceneNode);
+				UInterchangeSceneNode* ParentSceneNode = Cast<UInterchangeSceneNode>(Info.ParentNode);
+				ensure(ParentSceneNode);
+
+				TranslatorImpl.LODVariantSetNodeToLodMeshNodes.FindOrAdd(ParentSceneNode).Add(SceneNode);
 			}
 
 			// Skeleton joints are separate scene nodes in Interchange, so we need to emit that node hierarchy now
@@ -5643,6 +5664,7 @@ namespace UE::InterchangeUsdTranslator::Private
 			// to resolve all the bindings. The cached assignments come from AddMeshNode step, where
 			// we switch the active variant on the current stage and so get nice material bindings that
 			// resolve normally.
+			// TODO: We may not need this anymore, as we don't use population masks to open the LOD stages anymore
 			//
 			// Note that we can't even use the info cache here, because it wouldn't have cached info
 			// about the inactive LOD variants
@@ -5803,7 +5825,8 @@ namespace UE::InterchangeUsdTranslator::Private
 			// to resolve all the bindings. The cached assignments come from AddMeshNode step, where
 			// we switch the active variant on the current stage and so get nice material bindings that
 			// resolve normally.
-
+			// TODO: We may not need this anymore, as we don't use population masks to open the LOD stages anymore
+			//
 			// Note that we can't even use the info cache here, because it wouldn't have cached info
 			// about the inactive LOD variants
 			if (UsdUtils::FUsdPrimMaterialAssignmentInfo* CachedInfo = Impl.CachedMaterialAssignments.Find(PrimPath))
