USD: Refactor skeleton handling to not consider the skeleton itself as a root joint;

Implement alternate multiple root bone support mechanism;
Fix transforms directly on skeleton prims being fully ignored when importing with baking.

diff --git a/Engine/Plugins/Interchange/Extensions/OpenUSD/Source/Import/Private/InterchangeUsdTranslator.cpp b/Engine/Plugins/Interchange/Extensions/OpenUSD/Source/Import/Private/InterchangeUsdTranslator.cpp
--- a/Engine/Plugins/Interchange/Extensions/OpenUSD/Source/Import/Private/InterchangeUsdTranslator.cpp
+++ b/Engine/Plugins/Interchange/Extensions/OpenUSD/Source/Import/Private/InterchangeUsdTranslator.cpp
@@ -122,7 +122,6 @@ namespace UE::InterchangeUsdTranslator::Private
 #if USE_USD_SDK
 	const static FString AnimationPrefix = TEXT("\\Animation\\");
 	const static FString AnimationTrackPrefix = TEXT("\\AnimationTrack\\");
-	const static FString BonePrefix = TEXT("\\Bone\\");
 	const static FString CameraPrefix = TEXT("\\Camera\\");
 	const static FString GroomPrefix = TEXT("\\Groom\\");
 	const static FString LightPrefix = TEXT("\\Light\\");
@@ -149,11 +148,13 @@ namespace UE::InterchangeUsdTranslator::Private
 
 		UInterchangeBaseNode* ParentNode = nullptr;
 
+		FTransform SceneGlobalTransform;
+
 		TSharedPtr<UE::FUsdSkelCache> FurthestSkelCache;
 		TSharedPtr<FString> FurthestParentSkelRootPath;	   // Used to populate the skel cache
 		TSharedPtr<FString> ClosestParentSkelRootPath;
 		TSharedPtr<FString> BoundSkeletonPrimPath;
-		TSharedPtr<TArray<FString>> SkelJointNames;	   // Needed for skel mesh payloads
+		TSharedPtr<TArray<FString>> SkelJointUsdNames;	   // Needed for skel mesh payloads
 
 		// Nanite Assembly info
 		TSharedPtr<FString> ClosestNaniteAssemblySkelBindingPath;
@@ -3794,7 +3795,7 @@ namespace UE::InterchangeUsdTranslator::Private
 			MeshNode->SetPayLoadKey(PayloadKey, EInterchangeMeshPayLoadType::SKELETAL);
 			if (Info.BoundSkeletonPrimPath && !Info.BoundSkeletonPrimPath->IsEmpty())
 			{
-				MeshNode->SetSkeletonDependencyUid(MakeNodeUid(*Info.BoundSkeletonPrimPath));
+				MeshNode->SetSkeletonDependencyUid(MakeRootBoneNodeUid(*Info.BoundSkeletonPrimPath));
 			}
 
 			AddMorphTargetNodes(PrototypePrim, *this, *MeshNode, NodeContainer, Info);
@@ -4745,9 +4746,9 @@ namespace UE::InterchangeUsdTranslator::Private
 
 	const UInterchangeSkeletalAnimationTrackNode* AddSkeletalAnimationNode(
 		const UE::FUsdSkelSkeletonQuery& SkeletonQuery,
-		const TMap<FString, TPair<FString, int32>>& BoneToUidAndBoneIndex,
+		const TMap<FString, TPair<FString, int32>>& UsdBoneToUidAndUsdBoneIndex,
 		UInterchangeUSDTranslatorImpl& TranslatorImpl,
-		UInterchangeSceneNode& SkeletonPrimNode,
+		UInterchangeSceneNode& RootJointSceneNode,
 		UInterchangeBaseNodeContainer& NodeContainer,
 		const FTraversalInfo& Info
 	)
@@ -4786,7 +4787,7 @@ namespace UE::InterchangeUsdTranslator::Private
 
 		UInterchangeSkeletalAnimationTrackNode* SkelAnimNode = NewObject<UInterchangeSkeletalAnimationTrackNode>(&NodeContainer);
 		NodeContainer.SetupNode(SkelAnimNode, NodeUid, SkelAnimationName, EInterchangeNodeContainerType::TranslatedAsset);
-		SkelAnimNode->SetCustomSkeletonNodeUid(SkeletonPrimNode.GetUniqueID());
+		SkelAnimNode->SetCustomSkeletonNodeUid(RootJointSceneNode.GetUniqueID());
 
 		// TODO: Uncomment this whenever Interchange supports skeletal animation sections, because
 		// currently it seems that InterchangeLevelSequenceFactory.cpp doesn't even have the string "skel" anywhere.
@@ -4833,16 +4834,16 @@ namespace UE::InterchangeUsdTranslator::Private
 		TArray<FString> UsdJointOrder = AnimQuery.GetJointOrder();
 		for (const FString& FullAnimatedBoneName : UsdJointOrder)
 		{
-			const TPair<FString, int32>* FoundPair = BoneToUidAndBoneIndex.Find(FullAnimatedBoneName);
+			const TPair<FString, int32>* FoundPair = UsdBoneToUidAndUsdBoneIndex.Find(FullAnimatedBoneName);
 			if (!FoundPair)
 			{
 				continue;
 			}
 
 			const FString& BoneSceneNodeUid = FoundPair->Key;
-			int32 SkeletonOrderBoneIndex = FoundPair->Value;
+			int32 UsdSkeletonOrderBoneIndex = FoundPair->Value;
 
-			const FString BoneAnimPayloadKey = SkeletonPrimPath + TEXT("\\") + LexToString(SkeletonOrderBoneIndex);
+			const FString BoneAnimPayloadKey = SkeletonPrimPath + TEXT("\\") + LexToString(UsdSkeletonOrderBoneIndex);
 
 			// When retrieving the payload later, We'll need that bone's index within the Skeleton prim to index into the
 			// InUsdSkeletonQuery.ComputeJointLocalTransforms() results.
@@ -4972,19 +4973,11 @@ namespace UE::InterchangeUsdTranslator::Private
 		// By the time we get here we've already emitted a scene node for the skeleton prim itself, so we just
 		// need to emit a node hierarchy that mirrors the joints.
 
-		// Make the prim node into an Interchange joint/bone itself. By doing this we solve three issues:
-		//  - It becomes easy to identify our SkeletonDependencyUid when parsing Mesh nodes: It's just the skeleton prim path
-		//    (as opposed to having to target the translated node of the first root joint of the skeleton);
-		//  - We automatically handle USD skeletons with multiple root bones: We'll only ever have one "true"
-		//    root bone anyway: The SkeletonPrimNode itself;
-		//  - If a skeleton has no bones at all somehow, we'll still make one "bone" for it (this node).
-		SkeletonPrimNode.AddSpecializedType(UE::Interchange::FSceneNodeStaticData::GetJointSpecializeTypeString());
-
 		constexpr bool bResetCache = false;
 		SkeletonPrimNode.SetCustomBindPoseLocalTransform(&NodeContainer, FTransform::Identity, bResetCache);
 		SkeletonPrimNode.SetCustomTimeZeroLocalTransform(&NodeContainer, FTransform::Identity, bResetCache);
 
-		const FString SkeletonPrimNodeUid = SkeletonPrimNode.GetUniqueID();
+		const FString SkeletonPrimPath = Prim.GetPrimPath().GetString();
 
 #if WITH_EDITOR
 		// Convert the skeleton bones/joints into ConvertedData
@@ -4993,34 +4986,97 @@ namespace UE::InterchangeUsdTranslator::Private
 		const bool bEnsureSingleRootBone = false;
 		UsdToUnreal::FUsdSkeletonData ConvertedData;
 		const bool bSuccess = UsdToUnreal::ConvertSkeleton(SkelQuery, ConvertedData, bEnsureAtLeastOneBone, bEnsureSingleRootBone);
-		if (!bSuccess)
+		if (!bSuccess || ConvertedData.Bones.Num() == 0)
 		{
 			return;
 		}
 
+		// If we have multiple root bones, we need to add a single "true" root bone here, and then just add all the
+		// input root bones as children. We need to make sure our root bone has a unique name, but once that is done things
+		// should work out fine (including offsetting the vertex bone influence indices) because Interchange will compare these
+		// bone names with the bone names we return from GetSkeletalMeshPayloadData() and remap the indices by name
+		bool bHadMultipleRootBones = false;
+		{
+			TArray<int32> RootBoneIndices;
+			for (int32 BoneIndex = 0; BoneIndex < ConvertedData.Bones.Num(); ++BoneIndex)
+			{
+				const UsdToUnreal::FUsdSkeletonData::FBone& Bone = ConvertedData.Bones[BoneIndex];
+				if (Bone.ParentIndex == INDEX_NONE)
+				{
+					RootBoneIndices.Add(BoneIndex); 
+				}
+			}
+
+			if (RootBoneIndices.Num() > 1)
+			{
+				bHadMultipleRootBones = true;
+
+				UsdToUnreal::FUsdSkeletonData::FBone TrueRoot;
+				TrueRoot.Name = UE::Interchange::USD::RootBoneUidSuffix;
+
+				TrueRoot.ChildIndices = RootBoneIndices;
+				ConvertedData.Bones.Insert(TrueRoot, 0);
+
+				// This will also patch up our TrueRoot bone, which is great as it fixes
+				// up its ChildIndices too
+				for (UsdToUnreal::FUsdSkeletonData::FBone& Bone : ConvertedData.Bones)
+				{
+					Bone.ParentIndex += 1;
+					for (int32& ChildIndex : Bone.ChildIndices)
+					{
+						ChildIndex += 1;
+					}
+				}
+
+				// We need to fixup the true root's parentindex here though
+				ConvertedData.Bones[0].ParentIndex = INDEX_NONE;
+			}
+		}
+
 		// Maps from the USD-style full bone name (e.g. "shoulder/elbow/hand") to the Uid we used for
 		// the corresponding scene node, and the bone's index on the skeleton's joint order.
 		// We'll need this to parse skeletal animations, if any
-		TMap<FString, TPair<FString, int32>> BoneToUidAndBoneIndex;
+		TMap<FString, TPair<FString, int32>> UsdBoneToUidAndUsdBoneIndex;
+		UInterchangeSceneNode* RootBoneSceneNode = nullptr;
 
 		// Recursively traverse ConvertedData spawning the joint translated nodes
 		TFunction<void(int32, UInterchangeSceneNode&, const FString&)> RecursiveTraverseBones = nullptr;
 		RecursiveTraverseBones = [&RecursiveTraverseBones,
-								  &BoneToUidAndBoneIndex,
-								  &SkeletonPrimNodeUid,
+								  &UsdBoneToUidAndUsdBoneIndex,
+								  &SkeletonPrimPath,
 								  &ConvertedData,
-								  &NodeContainer	//
-		](int32 BoneIndex, UInterchangeSceneNode& ParentNode, const FString& BonePath)
+								  &NodeContainer,
+								  &RootBoneSceneNode,
+								  bHadMultipleRootBones
+		](int32 BoneIndex, UInterchangeSceneNode& ParentNode, const FString& UsdBonePath)
 		{
 			const UsdToUnreal::FUsdSkeletonData::FBone& Bone = ConvertedData.Bones[BoneIndex];
 
-			// Reconcatenate a full "bone path" here for uniqueness, because Bone.Name is just the name of this
-			// single bone/joint itself (e.g. "Elbow")
-			const FString ConcatBonePath = (BonePath.IsEmpty() ? TEXT("") : (BonePath + TEXT("/"))) + Bone.Name;
+			const bool bIsRootBone = BoneIndex == 0;
+
+			// Name of the individual bone, as-is from USD. Our "true root" becomes just "" here, as it doesn't exist in USD
+			const FString UsdBoneName = (bIsRootBone && bHadMultipleRootBones) ? TEXT("") : Bone.Name;
+
+			// Concatenate a full "bone path" here for uniqueness
+			// This now matches USD bone paths (e.g. "shoulder/elbow/hand")
+			const FString ConcatUsdBonePath = UsdBonePath.IsEmpty() ? UsdBoneName : UsdBonePath + TEXT("/") + UsdBoneName;
+
+			// Append our "true root" bone if we have one (e.g. shoulder/elbow/hand -> Root/shoulder/elbow/hand)
+			const FString FullConcatUsdBonePath = (bHadMultipleRootBones)
+				? RootBoneUidSuffix + (bIsRootBone ? TEXT("") : TEXT("/") + ConcatUsdBonePath)
+				: ConcatUsdBonePath;
 
-			// Putting the BonePrefix here avoids the pathological case where the user has skeleton child prims
-			// with names that match the joint names
-			const FString BoneNodeUid = MakeNodeUid(SkeletonPrimNodeUid + BonePrefix + ConcatBonePath);
+			// The UID for a root bone (whether "true root" or not) is always \Bone\<SkeletonPrimPath>/RootBoneUidSuffix (e.g. "\Bone\<SkeletonPrimPath>/Root")
+			// given our call to MakeRootBoneNodeUid().
+			// For any other bone, it becomes \Bone\<SkeletonPrimPath>/FullConcat (e.g. "\Bone\<SkeletonPrimPath>/<UsdRootBoneName>/Elbow/Hand")
+			//
+			// Note that this may look slightly odd for non-multiple-root-bone cases, because the root bone UID will be e.g. "\Bone\<SkeletonPrimPath>/Root", and its
+			// immediate child's will be "\Bone\<SkeletonPrimPath>/<UsdRootBoneName>/Elbow", so *not* "\Bone\<SkeletonPrimPath>/Root/Elbow".
+			// This simply because it's the easiest to manipulate/construct here, but either would have worked: These are just UIDs after all.
+			// The display labels are still just sensible bone names
+			const FString BoneNodeUid = bIsRootBone
+				? MakeRootBoneNodeUid(SkeletonPrimPath)
+				: MakeBoneNodeUid(SkeletonPrimPath, FullConcatUsdBonePath);
 
 			UInterchangeSceneNode* BoneNode = NewObject<UInterchangeSceneNode>(&NodeContainer);
 			NodeContainer.SetupNode(BoneNode, BoneNodeUid, Bone.Name, EInterchangeNodeContainerType::TranslatedScene, ParentNode.GetUniqueID());
@@ -5035,60 +5091,83 @@ namespace UE::InterchangeUsdTranslator::Private
 			BoneNode->SetCustomTimeZeroLocalTransform(&NodeContainer, Bone.LocalBindTransform, bResetCache);
 			BoneNode->SetCustomLocalTransform(&NodeContainer, Bone.LocalBindTransform, bResetCache);
 
-			BoneToUidAndBoneIndex.Add(ConcatBonePath, {BoneNodeUid, BoneIndex});
+			// These will be used to index into the baked transforms from USD, so we have to pretend our "true root" bone doesn't exist,
+			// if we added one
+			const int32 UsdBoneIndex = bHadMultipleRootBones ? BoneIndex - 1 : BoneIndex;
+
+			UsdBoneToUidAndUsdBoneIndex.Add(ConcatUsdBonePath, {BoneNodeUid, UsdBoneIndex});
+			if (bIsRootBone)
+			{
+				RootBoneSceneNode = BoneNode;
+			}
 
 			for (int32 ChildIndex : Bone.ChildIndices)
 			{
-				RecursiveTraverseBones(ChildIndex, *BoneNode, ConcatBonePath);
+				RecursiveTraverseBones(ChildIndex, *BoneNode, ConcatUsdBonePath);
 			}
 		};
 
-		// Start traversing from the root bones (we may have more than one, so check them all)
-		TSet<FString> UsedBoneNames;
-		for (int32 BoneIndex = 0; BoneIndex < ConvertedData.Bones.Num(); ++BoneIndex)
+		// Start traversing from the root bone (at this point we know we have exactly one)
 		{
-			const UsdToUnreal::FUsdSkeletonData::FBone& Bone = ConvertedData.Bones[BoneIndex];
-			UsedBoneNames.Add(Bone.Name);
-
-			if (Bone.ParentIndex == INDEX_NONE)
-			{
-				const FString BonePathRoot = TEXT("");
-				RecursiveTraverseBones(BoneIndex, SkeletonPrimNode, BonePathRoot);
-			}
+			const FString BonePathRoot = TEXT("");
+			const int32 BoneIndex = 0;
+			RecursiveTraverseBones(BoneIndex, SkeletonPrimNode, BonePathRoot);
+		}
+		if (!RootBoneSceneNode)
+		{
+			return;
 		}
 
 		// Interchange will abort parsing skeletons that don't have unique names for each bone. If the user has that
 		// on their actual skeleton, then that's just invalid data and we can just let it fail and emit the error message.
-		// However, we don't want to end up with duplicate bone names and fail to parse when the duplicate "bone" is due to
-		// how we actually use the Skeleton prim itself as the root, as that's our little "trick". In this case, here we
-		// just change the display text of the skeleton prim itself to be unique (which is used for the bone name)
-		const FString SkeletonPrimName = SkeletonPrimNode.GetDisplayLabel();
-		FString NewSkeletonPrimName = UsdUnreal::ObjectUtils::GetUniqueName(SkeletonPrimName, UsedBoneNames);
-		if (NewSkeletonPrimName != SkeletonPrimName)
+		// However, we don't want to end up with duplicate bone names and fail to parse when the duplicate bone is due to
+		// our extra true root bone... In this case, here we just change the display label of that node itself to be unique
+		// (which is used for the bone name)
+		if (bHadMultipleRootBones && RootBoneSceneNode)
 		{
-			SkeletonPrimNode.SetDisplayLabel(NewSkeletonPrimName);
+			TSet<FString> UsedBoneNames;
+			for (int32 BoneIndex = 1; BoneIndex < ConvertedData.Bones.Num(); ++BoneIndex)
+			{
+				const UsdToUnreal::FUsdSkeletonData::FBone& Bone = ConvertedData.Bones[BoneIndex];
+				UsedBoneNames.Add(Bone.Name);
+			}
+
+			FString NewSkeletonPrimName = UsdUnreal::ObjectUtils::GetUniqueName(UE::Interchange::USD::RootBoneUidSuffix, UsedBoneNames);
+			if (NewSkeletonPrimName != UE::Interchange::USD::RootBoneUidSuffix)
+			{
+				const FString RootBoneUid = MakeRootBoneNodeUid(SkeletonPrimPath);
+				RootBoneSceneNode->SetDisplayLabel(NewSkeletonPrimName);
+			}
 		}
 
 		// Handle SkelAnimation prims, if we have any bound for this Skeleton
 		const UInterchangeSkeletalAnimationTrackNode* SkelAnimNode = AddSkeletalAnimationNode(
 			SkelQuery,
-			BoneToUidAndBoneIndex,
+			UsdBoneToUidAndUsdBoneIndex,
 			*this,
-			SkeletonPrimNode,
+			*RootBoneSceneNode,
 			NodeContainer,
 			Info
 		);
 		if (SkelAnimNode)
 		{
-			SkeletonPrimNode.SetCustomAnimationAssetUidToPlay(SkelAnimNode->GetUniqueID());
+			RootBoneSceneNode->SetCustomAnimationAssetUidToPlay(SkelAnimNode->GetUniqueID());
 		}
 
 		// Cache our joint names in order, as this is needed when generating skeletal mesh payloads
-		Info.SkelJointNames = MakeShared<TArray<FString>>();
-		Info.SkelJointNames->Reserve(ConvertedData.Bones.Num());
+		Info.SkelJointUsdNames = MakeShared<TArray<FString>>();
+		Info.SkelJointUsdNames->Reserve(ConvertedData.Bones.Num());
 		for (const UsdToUnreal::FUsdSkeletonData::FBone& Bone : ConvertedData.Bones)
 		{
-			Info.SkelJointNames->Add(Bone.Name);
+			// The joint names that the skeletal mesh payload needs to retrieve should match exactly the USD bone
+			// structure (and so correspond to the vertex bone influence indices). We really don't want our extra
+			// "true root" bone to show up in there, or else we will mislead Interchange
+			if (bHadMultipleRootBones && Bone.ParentIndex == INDEX_NONE)
+			{
+				continue;
+			}
+
+			Info.SkelJointUsdNames->Add(Bone.Name);
 		}
 
 		// Prefer flagging the Skeleton prim itself as the BoundSkeletonPrimPath at this point in the hierarchy, even
@@ -5105,8 +5184,14 @@ namespace UE::InterchangeUsdTranslator::Private
 		}
 
 		{
+			// Cache the skeleton global transform, because if we need to bake our skeletal mesh later then we'll need to use this
+			// and it's awkward to have Interchange provide this to us natively, as the "skeleton prim node" here, which corresponds
+			// to the "first scene node parent of the root joint transform" doesn't really mean anything for the other formats...
+			const FTransform GlobalOffsetTransform;
+			SkeletonPrimNode.GetCustomGlobalTransform(&NodeContainer, GlobalOffsetTransform, Info.SceneGlobalTransform);
+
 			FWriteScopeLock ScopedInfoWriteLock{CachedTraversalInfoLock};
-			NodeUidToCachedTraversalInfo.Add(SkeletonPrimNodeUid, Info);
+			NodeUidToCachedTraversalInfo.Add(SkeletonPrimNode.GetUniqueID(), Info);
 		}
 #endif	  // WITH_EDITOR
 	}
@@ -5683,6 +5768,8 @@ namespace UE::InterchangeUsdTranslator::Private
 
 	bool GetSkeletalMeshPayloadData(
 		FString PayloadKey,
+		bool bBakeMeshes,
+		FTransform GlobalOffsetTransform,
 		UInterchangeUSDTranslatorImpl& Impl,
 		const UsdToUnreal::FUsdMeshConversionOptions& Options,
 		FMeshDescription& OutMeshDescription,
@@ -5717,6 +5804,8 @@ namespace UE::InterchangeUsdTranslator::Private
 		// const FString& MeshNodeUid = MeshPrefix + UsdUtils::GetPrototypePrimPath(Prim).GetString();
 		const FString& MeshNodeUid = MakeNodeUid(MeshPrefix + Prim.GetPrimPath().GetString());
 
+		UsdToUnreal::FUsdMeshConversionOptions OptionsCopy = Options;
+
 		// Read these variables from the data we cached during traversal for translation
 		TSharedPtr<TArray<FString>> JointNames = nullptr;
 		UE::FUsdSkelSkeletonQuery SkelQuery;
@@ -5743,12 +5832,17 @@ namespace UE::InterchangeUsdTranslator::Private
 			{
 				return false;
 			}
-			JointNames = SkeletonInfo->SkelJointNames;
+			JointNames = SkeletonInfo->SkelJointUsdNames;
 			if (!JointNames)
 			{
 				return false;
 			}
 
+			if (bBakeMeshes)
+			{
+				OptionsCopy.AdditionalTransform = SkeletonInfo->SceneGlobalTransform * GlobalOffsetTransform;
+			}
+
 			// Get skeletal mesh Nanite assembly data, if available
 			if (MeshInfo && MeshInfo->NaniteAssemblyTraversalResult)
 			{
@@ -5757,7 +5851,7 @@ namespace UE::InterchangeUsdTranslator::Private
 					// Add native instance data to the output description
 					GetNaniteAssemblyPayloadDataForPrims(
 						Impl.UsdStage,
-						Options,
+						OptionsCopy,
 						*TraversalResult,
 						SkeletonNodeUid,
 						*JointNames,
@@ -5769,7 +5863,7 @@ namespace UE::InterchangeUsdTranslator::Private
 						UE::FUsdPrim PointInstancerPrim = Impl.UsdStage.GetPrimAtPath(TopLevelPointInstancerPath);
 						GetNaniteAssemblyPayloadDataForPointInstancer(
 							PointInstancerPrim, 
-							Options,
+							OptionsCopy,
 							*TraversalResult,
 							SkeletonNodeUid, 
 							*JointNames, 
@@ -5805,7 +5899,7 @@ namespace UE::InterchangeUsdTranslator::Private
 			Prim,
 			TempMeshDescription,
 			TempMaterialInfo,
-			Options,
+			OptionsCopy,
 			&PxrSkelSkinningQueryPtr,
 			&PxrSkelSkeletonQueryPtr
 		);
@@ -5849,6 +5943,8 @@ namespace UE::InterchangeUsdTranslator::Private
 
 	bool GetMorphTargetPayloadData(
 		FString PayloadKey,
+		bool bBakeMeshes,
+		FTransform GlobalOffsetTransform,
 		UInterchangeUSDTranslatorImpl& Impl,
 		const UsdToUnreal::FUsdMeshConversionOptions& Options,
 		FMeshDescription& OutMeshDescription,
@@ -5895,7 +5991,16 @@ namespace UE::InterchangeUsdTranslator::Private
 		TArray<FString> UnusedJointNames;
 		TOptional<UE::Interchange::FNaniteAssemblyDescription> UnusedNaniteData;
 		const FString MeshPayloadKey = bIsLODMesh ? LODPrefix + MeshPrimPath : MeshPrimPath;
-		bool bConverted = GetSkeletalMeshPayloadData(MeshPayloadKey, Impl, Options, OutMeshDescription, UnusedJointNames, UnusedNaniteData);
+		bool bConverted = GetSkeletalMeshPayloadData(
+			MeshPayloadKey,
+			bBakeMeshes,
+			GlobalOffsetTransform,
+			Impl,
+			Options,
+			OutMeshDescription,
+			UnusedJointNames,
+			UnusedNaniteData
+		);
 		if (!bConverted || OutMeshDescription.IsEmpty())
 		{
 			return false;
@@ -6227,7 +6332,7 @@ namespace UE::InterchangeUsdTranslator::Private
 		{
 			const UE::Interchange::FAnimationPayloadQuery* Query = Queries[QueryIndex];
 
-			FString IndexStr = Query->PayloadKey.UniqueId.RightChop(SkeletonPrimPath.Len() + 1);	// Also skip the '\'
+			const FString& IndexStr = PayloadKeyTokens[1];
 			int32 JointIndex = INDEX_NONE;
 			bool bLexed = LexTryParseString(JointIndex, *IndexStr);
 			if (!bLexed)
@@ -6774,21 +6879,18 @@ TOptional<UE::Interchange::FMeshPayloadData> UInterchangeUSDTranslator::GetMeshP
 			// With baking, we want to first apply geomBindTransform, and then apply the skeleton's localToWorld transform. ConvertGeomMesh can 
 			// sort out the geomBindTransform (which should always be applied), so here we set the baking transform to the skeleton prim's transform if needed
 			bool bBakeMeshes = false;
-			FTransform RootJointGlobalTransform;
+			FTransform GlobalOffsetTransform;
 			PayloadAttributes.GetAttribute(UE::Interchange::FAttributeKey{MeshPayload::Attributes::BakeMeshes}, bBakeMeshes);
-			PayloadAttributes.GetAttribute(UE::Interchange::FAttributeKey{MeshPayload::Attributes::RootJointGlobalTransform}, RootJointGlobalTransform);
-
-			if (bBakeMeshes)
-			{
-				OptionsCopy.AdditionalTransform = RootJointGlobalTransform;
-			}
-
+			PayloadAttributes.GetAttribute(UE::Interchange::FAttributeKey{MeshPayload::Attributes::GlobalOffsetTransform}, GlobalOffsetTransform);
+			
 			// The skeletal mesh factory will already remap our bones by name from RetrieveAllSkeletalMeshPayloads() anyway,
 			// so we shouldn't offset anything ourselves
 			OptionsCopy.bOffsetJointIndicesForMultipleRootBones = false;
 
 			bSuccess = UE::InterchangeUsdTranslator::Private::GetSkeletalMeshPayloadData(
 				PayloadKey.UniqueId,
+				bBakeMeshes,
+				GlobalOffsetTransform,
 				*ImplPtr,
 				OptionsCopy,
 				MeshPayloadData.MeshDescription,
@@ -6801,17 +6903,14 @@ TOptional<UE::Interchange::FMeshPayloadData> UInterchangeUSDTranslator::GetMeshP
 		{
 			// See the ::SKELETAL case
 			bool bBakeMeshes = false;
-			FTransform RootJointGlobalTransform;
+			FTransform GlobalOffsetTransform;
 			PayloadAttributes.GetAttribute(UE::Interchange::FAttributeKey{MeshPayload::Attributes::BakeMeshes}, bBakeMeshes);
-			PayloadAttributes.GetAttribute(UE::Interchange::FAttributeKey{MeshPayload::Attributes::RootJointGlobalTransform}, RootJointGlobalTransform);
-
-			if (bBakeMeshes)
-			{
-				OptionsCopy.AdditionalTransform = RootJointGlobalTransform;
-			}
+			PayloadAttributes.GetAttribute(UE::Interchange::FAttributeKey{MeshPayload::Attributes::GlobalOffsetTransform}, GlobalOffsetTransform);
 
 			bSuccess = UE::InterchangeUsdTranslator::Private::GetMorphTargetPayloadData(
 				PayloadKey.UniqueId,
+				bBakeMeshes,
+				GlobalOffsetTransform,
 				*ImplPtr,
 				OptionsCopy,
 				MeshPayloadData.MeshDescription,
diff --git a/Engine/Plugins/Interchange/Runtime/Source/Common/Public/Usd/InterchangeUsdDefinitions.h b/Engine/Plugins/Interchange/Runtime/Source/Common/Public/Usd/InterchangeUsdDefinitions.h
--- a/Engine/Plugins/Interchange/Runtime/Source/Common/Public/Usd/InterchangeUsdDefinitions.h
+++ b/Engine/Plugins/Interchange/Runtime/Source/Common/Public/Usd/InterchangeUsdDefinitions.h
@@ -40,6 +40,16 @@ namespace UE::Interchange::USD
 	const FString UseTextureParameterSuffix = TEXT("Texture");
 	const FString UVIndexParameterSuffix = TEXT("UVIndex");
 
+	// In order to make it easy to find the skeleton root bone scene node UID from a prim path, for the root bone in particular
+	// we always use an UID like <skeleton prim path>/RootBoneUid (e.g. "/MySkelRoot/MySkeleton/Root")
+	//
+	// Note that this is unrelated to the bone's *name*, which is still named after the actual USD bone name
+	const FString RootBoneUidSuffix = TEXT("Root");
+
+	// Prefix we add to the scene nodes UIDs we artificially produce for skeleton joints, to make sure we don't accidentally
+	// produce an UID that is already used by a regular prim in the stage
+	const FString BonePrefix = TEXT("\\Bone\\");
+
 	// Flag indicating whether we should parse a UInterchangeShaderGraphNode on the InterchangeUSDPipeline.
 	// This is now only used for the compatible primvar code
 	const FString ParseMaterialIdentifier = TEXT("USD_MI_ParseMaterial");
@@ -94,4 +104,14 @@ namespace UE::Interchange::USD
 	{
 		return NodeUid + TEXT("_") + LexToString(FCrc::StrCrc32(*NodeUid));
 	}
+
+	inline FString MakeBoneNodeUid(const FString& SkeletonPrimPath, const FString& ConcatBonePath)
+	{
+		return MakeNodeUid(BonePrefix + SkeletonPrimPath + TEXT("/") + ConcatBonePath);
+	}
+
+	inline FString MakeRootBoneNodeUid(const FString& SkeletonPrimPath)
+	{
+		return MakeBoneNodeUid(SkeletonPrimPath, RootBoneUidSuffix);
+	}
 }	 // namespace UE::Interchange::USD
diff --git a/Engine/Plugins/Interchange/Runtime/Source/Import/Private/Mesh/InterchangeSkeletalMeshFactory.cpp b/Engine/Plugins/Interchange/Runtime/Source/Import/Private/Mesh/InterchangeSkeletalMeshFactory.cpp
--- a/Engine/Plugins/Interchange/Runtime/Source/Import/Private/Mesh/InterchangeSkeletalMeshFactory.cpp
+++ b/Engine/Plugins/Interchange/Runtime/Source/Import/Private/Mesh/InterchangeSkeletalMeshFactory.cpp
@@ -1067,6 +1067,7 @@ void UInterchangeSkeletalMeshFactory::CreatePayloadTasks(const FImportAssetObjec
 
 		PayloadAttributes.RegisterAttribute(UE::Interchange::FAttributeKey{ UE::Interchange::MeshPayload::Attributes::BakeMeshes }, bBakeMeshes);
 		PayloadAttributes.RegisterAttribute(UE::Interchange::FAttributeKey{ UE::Interchange::MeshPayload::Attributes::RootJointGlobalTransform }, RootJointGlobalTransform);
+		PayloadAttributes.RegisterAttribute(UE::Interchange::FAttributeKey{ UE::Interchange::MeshPayload::Attributes::GlobalOffsetTransform }, GlobalOffsetTransform);
 
 		for (const FMeshNodeContext& MeshNodeContext : ImportAssetObjectLODData.MeshNodeContexts)
 		{
diff --git a/Engine/Plugins/Interchange/Runtime/Source/Nodes/Public/InterchangeMeshDefinitions.h b/Engine/Plugins/Interchange/Runtime/Source/Nodes/Public/InterchangeMeshDefinitions.h
--- a/Engine/Plugins/Interchange/Runtime/Source/Nodes/Public/InterchangeMeshDefinitions.h
+++ b/Engine/Plugins/Interchange/Runtime/Source/Nodes/Public/InterchangeMeshDefinitions.h
@@ -57,8 +57,16 @@ namespace UE
 		{
 			namespace Attributes
 			{
+				/** Global transform of the scene node that corresponds to the root joint of the skeleton. Already includes GlobalOffsetTransform */
 				const FString RootJointGlobalTransform = TEXT("__Payload__RootJointGlobalTransform");
+
+				/** Global transform for the mesh scene node. Already includes GlobalOffsetTransform */
 				const FString MeshGlobalTransform = TEXT("__Payload__MeshGlobalTransform");
+
+				/** Additional transform to be added last during the mesh baking (so as if an additional parent transform of the top level actor) */
+				const FString GlobalOffsetTransform = TEXT("__Payload__GlobalOffsetTransform");
+
+				/** Whether to perform mesh baking or not */
 				const FString BakeMeshes = TEXT("__Payload__BakeMeshes");
 			}
 		}
