USD Interchange: Add error message when we fail to open a USD stage.

diff --git a/Engine/Plugins/Interchange/Extensions/OpenUSD/Source/Import/Private/InterchangeUsdTranslator.cpp b/Engine/Plugins/Interchange/Extensions/OpenUSD/Source/Import/Private/InterchangeUsdTranslator.cpp
--- a/Engine/Plugins/Interchange/Extensions/OpenUSD/Source/Import/Private/InterchangeUsdTranslator.cpp
+++ b/Engine/Plugins/Interchange/Extensions/OpenUSD/Source/Import/Private/InterchangeUsdTranslator.cpp
@@ -6666,6 +6666,13 @@ bool UInterchangeUSDTranslator::Translate(UInterchangeBaseNodeContainer& NodeCon
 
 		if (!StageToImport)
 		{
+			UInterchangeResultError_Generic* Message = ImplPtr->ResultsContainer->Add<UInterchangeResultError_Generic>();
+			Message->SourceAssetName = SourceData->GetFilename();
+			Message->Text = FText::Format(
+				LOCTEXT("FailedToOpenStage", "Failed to get a valid stage from file '{0}'. Does the file exist and is it a valid USD file?"),
+				FText::FromString(SourceData->GetFilename())
+			);
+
 			return false;
 		}
 
