USD: Fix duplicate LOCTEXT key.

#rnx

diff --git a/Engine/Plugins/Interchange/Extensions/OpenUSD/Source/Import/Private/InterchangeUsdTranslator.cpp b/Engine/Plugins/Interchange/Extensions/OpenUSD/Source/Import/Private/InterchangeUsdTranslator.cpp
--- a/Engine/Plugins/Interchange/Extensions/OpenUSD/Source/Import/Private/InterchangeUsdTranslator.cpp
+++ b/Engine/Plugins/Interchange/Extensions/OpenUSD/Source/Import/Private/InterchangeUsdTranslator.cpp
@@ -5289,7 +5289,7 @@ namespace UE::InterchangeUsdTranslator::Private
 			Message->SourceAssetName = TranslatorImpl.UsdStage.GetRootLayer().GetRealPath();
 			Message->Text = FText::Format(
 				LOCTEXT(
-					"UnsupportedSchemaType",
+					"NestedGprims",
 					"Prim '{0}' is a Gprim child of another Gprim. This configuration is not recommended or supported very well, and may lead to unexpected results."
 				),
 				FText::FromString(Prim.GetPrimPath().GetString())
