Interchange USD: Set some min/max for the SubdivisionLevel property editor similar to USD legacy.

diff --git a/Engine/Plugins/Interchange/Extensions/OpenUSD/Source/Import/Private/InterchangeUSDPipeline.cpp b/Engine/Plugins/Interchange/Extensions/OpenUSD/Source/Import/Private/InterchangeUSDPipeline.cpp
--- a/Engine/Plugins/Interchange/Extensions/OpenUSD/Source/Import/Private/InterchangeUSDPipeline.cpp
+++ b/Engine/Plugins/Interchange/Extensions/OpenUSD/Source/Import/Private/InterchangeUSDPipeline.cpp
@@ -1511,6 +1511,21 @@ void UInterchangeUsdPipeline::ExecutePostFactoryPipeline(
 #endif	  // USE_USD_SDK && WITH_EDITORONLY_DATA
 }
 
+void UInterchangeUsdPipeline::AdjustSettingsForContext(const FInterchangePipelineContextParams& ContextParams)
+{
+	Super::AdjustSettingsForContext(ContextParams);
+
+	static IConsoleVariable* MaxSubdivCvar = IConsoleManager::Get().FindConsoleVariable(TEXT("USD.Subdiv.MaxSubdivLevel"));
+	const int32 MaxSubdivLevel = FMath::Max(MaxSubdivCvar ? MaxSubdivCvar->GetInt() : 6, 1);
+	const FString MaxSubdivLevelString = LexToString(MaxSubdivLevel);
+
+	if (FProperty* Property = StaticClass()->FindPropertyByName(GET_MEMBER_NAME_CHECKED(UInterchangeUsdPipeline, SubdivisionLevel)))
+	{
+		Property->SetMetaData(TEXT("UIMax"), *MaxSubdivLevelString);
+		Property->SetMetaData(TEXT("ClampMax"), *MaxSubdivLevelString);
+	}
+}
+
 void UInterchangeUsdPipeline::ExecutePipeline(
 	UInterchangeBaseNodeContainer* NodeContainer,
 	const TArray<UInterchangeSourceData*>& InSourceDatas,
diff --git a/Engine/Plugins/Interchange/Extensions/OpenUSD/Source/Import/Public/InterchangeUSDPipeline.h b/Engine/Plugins/Interchange/Extensions/OpenUSD/Source/Import/Public/InterchangeUSDPipeline.h
--- a/Engine/Plugins/Interchange/Extensions/OpenUSD/Source/Import/Public/InterchangeUSDPipeline.h
+++ b/Engine/Plugins/Interchange/Extensions/OpenUSD/Source/Import/Public/InterchangeUSDPipeline.h
@@ -38,7 +38,7 @@ public:
 	 * Subdivision level to use for all subdivision meshes on the opened stage. 0 means "don't subdivide".
 	 * The maximum level of subdivision allowed can be configured via the 'USD.Subdiv.MaxSubdivLevel' cvar.
 	 */
-	UPROPERTY(EditAnywhere, BlueprintReadWrite, Category = "USD")
+	UPROPERTY(EditAnywhere, BlueprintReadWrite, Category = "USD", meta = (ClampMin = "0", ClampMax = "6", UIMin = "0", UIMax = "6"))
 	int32 SubdivisionLevel;
 
 	/**
@@ -62,6 +62,8 @@ public:
 
 	virtual void ExecutePostFactoryPipeline(const UInterchangeBaseNodeContainer* BaseNodeContainer, const FString& NodeKey, UObject* CreatedAsset, bool bIsAReimport) override;
 
+	virtual void AdjustSettingsForContext(const FInterchangePipelineContextParams& ContextParams) override;
+
 protected:
 	virtual void ExecutePipeline(
 		UInterchangeBaseNodeContainer* BaseNodeContainer,
