Fixed crash when Interchange stack names have been modified.

#rnx

diff --git a/Engine/Plugins/Experimental/Enterprise/DatasmithInterchange/Source/DatasmithInterchange/Private/DatasmithInterchangeModule.cpp b/Engine/Plugins/Experimental/Enterprise/DatasmithInterchange/Source/DatasmithInterchange/Private/DatasmithInterchangeModule.cpp
--- a/Engine/Plugins/Experimental/Enterprise/DatasmithInterchange/Source/DatasmithInterchange/Private/DatasmithInterchangeModule.cpp
+++ b/Engine/Plugins/Experimental/Enterprise/DatasmithInterchange/Source/DatasmithInterchange/Private/DatasmithInterchangeModule.cpp
@@ -80,14 +80,36 @@ public:
 		TranslatorPipeplines.Pipelines.Add(FSoftObjectPath(TEXT("/DatasmithInterchange/InterchangeDatasmithDefault.InterchangeDatasmithDefault")));
 		{
 			FInterchangeImportSettings& ImportSettings = FInterchangeProjectSettingsUtils::GetMutableDefaultImportSettings(false);
-			FInterchangePipelineStack& PipelineStack = ImportSettings.PipelineStacks[TEXT("Assets")];
-			PipelineStack.PerTranslatorPipelines.Add(TranslatorPipeplines);
+
+			const FName DefaultPipelineStack = ImportSettings.DefaultPipelineStack;
+			if (ImportSettings.PipelineStacks.Contains(DefaultPipelineStack))
+			{
+				FInterchangePipelineStack& PipelineStack = ImportSettings.PipelineStacks[DefaultPipelineStack];
+				PipelineStack.PerTranslatorPipelines.Add(TranslatorPipeplines);
+
+				UE_LOG(LogInterchangeDatasmith, Log, TEXT("InterchangeDatasmith's pipelines added to default Asset pipeline stack, %s"), *DefaultPipelineStack.ToString());
+			}
+			else
+			{
+				UE_LOG(LogInterchangeDatasmith, Warning, TEXT("Cannot find Interchange default Asset's stack to add pipelines. InterchangeDatasmith asset import may not work properly."));
+			}
 		}
 
 		{
 			FInterchangeImportSettings& ImportSettings = FInterchangeProjectSettingsUtils::GetMutableDefaultImportSettings(true);
-			FInterchangePipelineStack& PipelineStack = ImportSettings.PipelineStacks[TEXT("Scene")];
-			PipelineStack.PerTranslatorPipelines.Add(TranslatorPipeplines);
+
+			const FName DefaultPipelineStack = ImportSettings.DefaultPipelineStack;
+			if (ImportSettings.PipelineStacks.Contains(DefaultPipelineStack))
+			{
+				FInterchangePipelineStack& PipelineStack = ImportSettings.PipelineStacks[DefaultPipelineStack];
+				PipelineStack.PerTranslatorPipelines.Add(TranslatorPipeplines);
+
+				UE_LOG(LogInterchangeDatasmith, Log, TEXT("InterchangeDatasmith's pipelines added to default Scene pipeline stack, %s"), *DefaultPipelineStack.ToString());
+			}
+			else
+			{
+				UE_LOG(LogInterchangeDatasmith, Warning, TEXT("Cannot find Interchange default Scene's stack to add pipelines. InterchangeDatasmith level import may not work properly."));
+			}
 		}
 
 		FDatasmithReferenceMaterialManager::Create();
diff --git a/Engine/Plugins/Interchange/Extensions/OpenUSD/Source/Import/Private/InterchangeOpenUSDImportModule.cpp b/Engine/Plugins/Interchange/Extensions/OpenUSD/Source/Import/Private/InterchangeOpenUSDImportModule.cpp
--- a/Engine/Plugins/Interchange/Extensions/OpenUSD/Source/Import/Private/InterchangeOpenUSDImportModule.cpp
+++ b/Engine/Plugins/Interchange/Extensions/OpenUSD/Source/Import/Private/InterchangeOpenUSDImportModule.cpp
@@ -10,6 +10,8 @@
 #include "Misc/CoreDelegates.h"
 #include "Modules/ModuleManager.h"
 
+DEFINE_LOG_CATEGORY_STATIC(LogInterchangeUSDModule, Log, All);
+
 void FInterchangeOpenUSDImportModule::StartupModule()
 {
 	using namespace UE::Interchange;
@@ -40,28 +42,48 @@ void FInterchangeOpenUSDImportModule::StartupModule()
 
 		// Asset import pipelines
 		{
-			FInterchangeTranslatorPipelines TranslatorPipelines;
-			TranslatorPipelines.Translator = TranslatorClassPath;
-			TranslatorPipelines.Pipelines.Add(FSoftObjectPath(TEXT("/InterchangeOpenUSD/Pipelines/DefaultUSDAssetsPipeline.DefaultUSDAssetsPipeline")));
-			TranslatorPipelines.Pipelines.Add(FSoftObjectPath(TEXT("/InterchangeOpenUSD/Pipelines/DefaultUSDPipelineAssetImport.DefaultUSDPipelineAssetImport")));
-			TranslatorPipelines.Pipelines.Add(FSoftObjectPath(TEXT("/Interchange/Pipelines/DefaultMaterialXPipeline.DefaultMaterialXPipeline")));
-			TranslatorPipelines.Pipelines.Add(FSoftObjectPath(TEXT("/Interchange/Pipelines/DefaultAudioPipeline.DefaultAudioPipeline")));
-
-			FInterchangePipelineStack& PipelineStack = AssetImportSettings.PipelineStacks[TEXT("Assets")];
-			PipelineStack.PerTranslatorPipelines.Add(TranslatorPipelines);
+			const FName DefaultPipelineStack = AssetImportSettings.DefaultPipelineStack;
+			if (AssetImportSettings.PipelineStacks.Contains(DefaultPipelineStack))
+			{
+				FInterchangeTranslatorPipelines TranslatorPipelines;
+				TranslatorPipelines.Translator = TranslatorClassPath;
+				TranslatorPipelines.Pipelines.Add(FSoftObjectPath(TEXT("/InterchangeOpenUSD/Pipelines/DefaultUSDAssetsPipeline.DefaultUSDAssetsPipeline")));
+				TranslatorPipelines.Pipelines.Add(FSoftObjectPath(TEXT("/InterchangeOpenUSD/Pipelines/DefaultUSDPipelineAssetImport.DefaultUSDPipelineAssetImport")));
+				TranslatorPipelines.Pipelines.Add(FSoftObjectPath(TEXT("/Interchange/Pipelines/DefaultMaterialXPipeline.DefaultMaterialXPipeline")));
+				TranslatorPipelines.Pipelines.Add(FSoftObjectPath(TEXT("/Interchange/Pipelines/DefaultAudioPipeline.DefaultAudioPipeline")));
+
+				FInterchangePipelineStack& PipelineStack = AssetImportSettings.PipelineStacks[DefaultPipelineStack];
+				PipelineStack.PerTranslatorPipelines.Add(TranslatorPipelines);
+				
+				UE_LOG(LogInterchangeUSDModule, Log, TEXT("InterchangeUSD's pipelines added to default Asset pipeline stack, %s"), *DefaultPipelineStack.ToString());
+			}
+			else
+			{
+				UE_LOG(LogInterchangeUSDModule, Warning, TEXT("Cannot find Interchange default Asset's stack to add pipelines. InterchangeUSD asset import may not work properly."));
+			}
 		}
 
 		// Scene import pipelines
 		{
-			FInterchangeTranslatorPipelines TranslatorPipelines;
-			TranslatorPipelines.Translator = TranslatorClassPath;
-			TranslatorPipelines.Pipelines.Add(FSoftObjectPath(TEXT("/InterchangeOpenUSD/Pipelines/DefaultUSDSceneAssetsPipeline.DefaultUSDSceneAssetsPipeline")));
-			TranslatorPipelines.Pipelines.Add(FSoftObjectPath(TEXT("/InterchangeOpenUSD/Pipelines/DefaultUSDSceneLevelPipeline.DefaultUSDSceneLevelPipeline")));
-			TranslatorPipelines.Pipelines.Add(FSoftObjectPath(TEXT("/InterchangeOpenUSD/Pipelines/DefaultUSDPipeline.DefaultUSDPipeline")));
-			TranslatorPipelines.Pipelines.Add(FSoftObjectPath(TEXT("/Interchange/Pipelines/DefaultMaterialXPipeline.DefaultMaterialXPipeline")));
-
-			FInterchangePipelineStack& PipelineStack = SceneImportSettings.PipelineStacks[TEXT("Scene")];
-			PipelineStack.PerTranslatorPipelines.Add(TranslatorPipelines);
+			const FName DefaultPipelineStack = SceneImportSettings.DefaultPipelineStack;
+			if (SceneImportSettings.PipelineStacks.Contains(DefaultPipelineStack))
+			{
+				FInterchangeTranslatorPipelines TranslatorPipelines;
+				TranslatorPipelines.Translator = TranslatorClassPath;
+				TranslatorPipelines.Pipelines.Add(FSoftObjectPath(TEXT("/InterchangeOpenUSD/Pipelines/DefaultUSDSceneAssetsPipeline.DefaultUSDSceneAssetsPipeline")));
+				TranslatorPipelines.Pipelines.Add(FSoftObjectPath(TEXT("/InterchangeOpenUSD/Pipelines/DefaultUSDSceneLevelPipeline.DefaultUSDSceneLevelPipeline")));
+				TranslatorPipelines.Pipelines.Add(FSoftObjectPath(TEXT("/InterchangeOpenUSD/Pipelines/DefaultUSDPipeline.DefaultUSDPipeline")));
+				TranslatorPipelines.Pipelines.Add(FSoftObjectPath(TEXT("/Interchange/Pipelines/DefaultMaterialXPipeline.DefaultMaterialXPipeline")));
+
+				FInterchangePipelineStack& PipelineStack = SceneImportSettings.PipelineStacks[DefaultPipelineStack];
+				PipelineStack.PerTranslatorPipelines.Add(TranslatorPipelines);
+
+				UE_LOG(LogInterchangeUSDModule, Log, TEXT("InterchangeUSD's pipelines added to default Scene pipeline stack, %s"), *DefaultPipelineStack.ToString());
+			}
+			else
+			{
+				UE_LOG(LogInterchangeUSDModule, Warning, TEXT("Cannot find Interchange default Scene's stack to add pipelines. InterchangeUSD level import may not work properly."));
+			}
 		}
 
 		// Asset import dialog overrides (we want to show the import and reimport dialog for all USD imports, of all asset types)
diff --git a/Engine/Plugins/Interchange/Extensions/OpenVDB/Source/Import/Private/InterchangeOpenVDBImportModule.cpp b/Engine/Plugins/Interchange/Extensions/OpenVDB/Source/Import/Private/InterchangeOpenVDBImportModule.cpp
--- a/Engine/Plugins/Interchange/Extensions/OpenVDB/Source/Import/Private/InterchangeOpenVDBImportModule.cpp
+++ b/Engine/Plugins/Interchange/Extensions/OpenVDB/Source/Import/Private/InterchangeOpenVDBImportModule.cpp
@@ -44,28 +44,66 @@ void FInterchangeOpenVDBImportModule::StartupModule()
 		}
 
 		// Scene import pipeline stacks
-		FInterchangeImportSettings& SceneImportSettings = ProjectSettings->SceneImportSettings;
-		FInterchangePipelineStack& PipelineStack = SceneImportSettings.PipelineStacks[TEXT("Scene")];
-		PipelineStack.PerTranslatorPipelines.Add(TranslatorPipelines);
-
-		// Asset import pipeline stacks
-		FInterchangeContentImportSettings& AssetImportSettings = ProjectSettings->ContentImportSettings;
-		FInterchangePipelineStack& AssetStack = AssetImportSettings.PipelineStacks[TEXT("Assets")];
-		AssetStack.PerTranslatorPipelines.Add(TranslatorPipelines);
-		FInterchangePipelineStack& TextureStack = AssetImportSettings.PipelineStacks[TEXT("Textures")];
-		TextureStack.PerTranslatorPipelines.Add(TranslatorPipelines);
-
-		// Asset import show dialog override (otherwise it doesn't show the import options dialog as OpenVDBs are Texture type assets)
 		{
-			FInterchangeDialogOverride& DialogOverrides = AssetImportSettings.ShowImportDialogOverride.FindOrAdd(
-				EInterchangeTranslatorAssetType::Textures
-			);
-
-			FInterchangePerTranslatorDialogOverride ImportDialogOverride;
-			ImportDialogOverride.Translator = TranslatorClassPath;
-			ImportDialogOverride.bShowImportDialog = true;
-			ImportDialogOverride.bShowReimportDialog = true;
-			DialogOverrides.PerTranslatorImportDialogOverride.Add(ImportDialogOverride);
+			FInterchangeImportSettings& SceneImportSettings = ProjectSettings->SceneImportSettings;
+			const FName DefaultPipelineStack = SceneImportSettings.DefaultPipelineStack;
+
+			if (SceneImportSettings.PipelineStacks.Contains(DefaultPipelineStack))
+			{
+				FInterchangePipelineStack& PipelineStack = SceneImportSettings.PipelineStacks[DefaultPipelineStack];
+				PipelineStack.PerTranslatorPipelines.Add(TranslatorPipelines);
+
+				UE_LOG(LogInterchangeOpenVDBImport, Log, TEXT("InterchangeOpenVDB's pipelines added to default Scene pipeline stack, %s"), *DefaultPipelineStack.ToString());
+			}
+			else
+			{
+				UE_LOG(LogInterchangeOpenVDBImport, Warning, TEXT("Cannot find Interchange default Scene's stack to add pipelines. InterchangeOpenVDB level import may not work properly."));
+			}
+		}
+
+		// Asset and Texture import pipeline stacks
+		{
+			FInterchangeContentImportSettings& AssetImportSettings = ProjectSettings->ContentImportSettings;
+
+			const FName DefaultPipelineStack = AssetImportSettings.DefaultPipelineStack;
+			if (AssetImportSettings.PipelineStacks.Contains(DefaultPipelineStack))
+			{
+				FInterchangePipelineStack& AssetStack = AssetImportSettings.PipelineStacks[DefaultPipelineStack];
+				AssetStack.PerTranslatorPipelines.Add(TranslatorPipelines);
+
+				UE_LOG(LogInterchangeOpenVDBImport, Log, TEXT("InterchangeOpenVDB's pipelines added to default Asset pipeline stack, %s"), *DefaultPipelineStack.ToString());
+			}
+			else
+			{
+				UE_LOG(LogInterchangeOpenVDBImport, Warning, TEXT("Cannot find Interchange default Asset's stack to add pipelines. InterchangeOpenVDB asset import may not work properly."));
+			}
+
+			if (AssetImportSettings.DefaultPipelineStackOverride.Contains(EInterchangeTranslatorAssetType::Textures))
+			{
+				const FName DefaultTextureStack = AssetImportSettings.DefaultPipelineStackOverride[EInterchangeTranslatorAssetType::Textures];
+				if (AssetImportSettings.PipelineStacks.Contains(DefaultTextureStack))
+				{
+					FInterchangePipelineStack& TextureStack = AssetImportSettings.PipelineStacks[DefaultTextureStack];
+					TextureStack.PerTranslatorPipelines.Add(TranslatorPipelines);
+
+					// Asset import show dialog override (otherwise it doesn't show the import options dialog as OpenVDBs are Texture type assets)
+					{
+						FInterchangeDialogOverride& DialogOverrides = AssetImportSettings.ShowImportDialogOverride.FindOrAdd(
+							EInterchangeTranslatorAssetType::Textures
+						);
+
+						FInterchangePerTranslatorDialogOverride ImportDialogOverride;
+						ImportDialogOverride.Translator = TranslatorClassPath;
+						ImportDialogOverride.bShowImportDialog = true;
+						ImportDialogOverride.bShowReimportDialog = true;
+						DialogOverrides.PerTranslatorImportDialogOverride.Add(ImportDialogOverride);
+					}
+				}
+				else
+				{
+					UE_LOG(LogInterchangeOpenVDBImport, Warning, TEXT("Cannot find Interchange default Textures's stack to add pipelines. InterchangeOpenVDB Texture import may not work properly."));
+				}
+			}
 		}
 #endif
 	};
