USD Interchange: Tweak tooltips to say only default attribute opinions are hoisted as Interchange node attributes

diff --git a/Engine/Plugins/Interchange/Extensions/OpenUSD/Source/Import/Public/InterchangeUsdTranslator.h b/Engine/Plugins/Interchange/Extensions/OpenUSD/Source/Import/Public/InterchangeUsdTranslator.h
--- a/Engine/Plugins/Interchange/Extensions/OpenUSD/Source/Import/Public/InterchangeUsdTranslator.h
+++ b/Engine/Plugins/Interchange/Extensions/OpenUSD/Source/Import/Public/InterchangeUsdTranslator.h
@@ -57,7 +57,7 @@ public:
 	FUsdStageOptions StageOptions;
 
 	/**
-	 * Whether to convert USD prim attributes into translated node user attributes, ultimately ending up imported as asset and actor metadata.
+	 * Whether to convert USD prim default attribute opinions (i.e. non-animated) into translated node user attributes, ultimately ending up imported as asset and actor metadata.
 	 *
 	 * Note that AttributeRegexFilter may be needed in order to select only the desired metadata, as otherwise regular attributes like 
 	 * 'points' and 'faceVertexCounts' will also get translated into Interchange metadata, which could be expensive.
