Interchange USD:
- Create Material Instance for TwoSided material instead of duplicating them

diff --git a/Engine/Plugins/Interchange/Extensions/OpenUSD/Source/Import/Private/InterchangeUsdTranslator.cpp b/Engine/Plugins/Interchange/Extensions/OpenUSD/Source/Import/Private/InterchangeUsdTranslator.cpp
--- a/Engine/Plugins/Interchange/Extensions/OpenUSD/Source/Import/Private/InterchangeUsdTranslator.cpp
+++ b/Engine/Plugins/Interchange/Extensions/OpenUSD/Source/Import/Private/InterchangeUsdTranslator.cpp
@@ -1531,12 +1531,12 @@ namespace UE::InterchangeUsdTranslator::Private
 		return Result;
 	}
 
-	FString GetOrCreateTwoSidedShaderGraphNode(const FString& OneSidedShaderGraphNodeUid, UInterchangeBaseNodeContainer& NodeContainer)
+	FString GetOrCreateTwoSidedMaterialInstanceNode(const FString& OneSidedShaderGraphNodeUid, UInterchangeBaseNodeContainer& NodeContainer)
 	{
 		FString TwoSidedUid = OneSidedShaderGraphNodeUid + TwoSidedSuffix;
 
 		// We already created this, just return it
-		if (const UInterchangeShaderGraphNode* Node = Cast<UInterchangeShaderGraphNode>(NodeContainer.GetNode(TwoSidedUid)))
+		if (Cast<UInterchangeMaterialInstanceNode>(NodeContainer.GetNode(TwoSidedUid)))
 		{
 			return TwoSidedUid;
 		}
@@ -1565,14 +1565,11 @@ namespace UE::InterchangeUsdTranslator::Private
 		}
 
 		FString TwoSidedNodeName = OneSidedNode->GetDisplayLabel() + TwoSidedSuffix;
-		UInterchangeShaderGraphNode* TwoSidedNode = NewObject<UInterchangeShaderGraphNode>(&NodeContainer);
-
-		UInterchangeBaseNode::CopyStorage(OneSidedNode, TwoSidedNode);
+		UInterchangeMaterialInstanceNode* TwoSidedNode = NewObject<UInterchangeMaterialInstanceNode>(&NodeContainer);
 
 		NodeContainer.SetupNode(TwoSidedNode, TwoSidedUid, TwoSidedNodeName, EInterchangeNodeContainerType::TranslatedAsset);
-
+		TwoSidedNode->SetCustomParent(OneSidedShaderGraphNodeUid);
 		TwoSidedNode->SetCustomTwoSided(true);
-		TwoSidedNode->SetCustomTwoSidedTransmission(true);
 
 		return TwoSidedUid;
 	}
@@ -2903,7 +2900,7 @@ namespace UE::InterchangeUsdTranslator::Private
 			FString MaterialXMaterialUid = GetMaterialXMaterialUid(MaterialPrimName, NodeContainer);
 			if (bForceTwoSided)
 			{
-				MaterialXMaterialUid = GetOrCreateTwoSidedShaderGraphNode(MaterialXMaterialUid, NodeContainer);
+				MaterialXMaterialUid = GetOrCreateTwoSidedMaterialInstanceNode(MaterialXMaterialUid, NodeContainer);
 			}
 
 			if (GeomPropValues.Num() > 0)
