Interchange USD: Remove the layer extension from the pseudoroot actor display label

diff --git a/Engine/Plugins/Interchange/Extensions/OpenUSD/Source/Import/Private/InterchangeUsdTranslator.cpp b/Engine/Plugins/Interchange/Extensions/OpenUSD/Source/Import/Private/InterchangeUsdTranslator.cpp
--- a/Engine/Plugins/Interchange/Extensions/OpenUSD/Source/Import/Private/InterchangeUsdTranslator.cpp
+++ b/Engine/Plugins/Interchange/Extensions/OpenUSD/Source/Import/Private/InterchangeUsdTranslator.cpp
@@ -5216,7 +5216,7 @@ namespace UE::InterchangeUsdTranslator::Private
 		FString DisplayLabel;
 		if (Prim.IsPseudoRoot())
 		{
-			DisplayLabel = Prim.GetStage().GetRootLayer().GetDisplayName();
+			DisplayLabel = FPaths::GetBaseFilename(Prim.GetStage().GetRootLayer().GetDisplayName());
 		}
 		else
 		{
