USD: Forward all USD SDK and USD Interchange errors to the USD error utils;

Display USD error utils messages as Interchange results

diff --git a/Engine/Plugins/Interchange/Extensions/OpenUSD/Source/Import/Private/InterchangeUsdTranslator.cpp b/Engine/Plugins/Interchange/Extensions/OpenUSD/Source/Import/Private/InterchangeUsdTranslator.cpp
--- a/Engine/Plugins/Interchange/Extensions/OpenUSD/Source/Import/Private/InterchangeUsdTranslator.cpp
+++ b/Engine/Plugins/Interchange/Extensions/OpenUSD/Source/Import/Private/InterchangeUsdTranslator.cpp
@@ -6,6 +6,7 @@
 #include "Objects/USDSchemaTranslator.h"
 #include "UnrealUSDWrapper.h"
 #include "USDConversionUtils.h"
+#include "USDErrorUtils.h"
 #include "USDGeomMeshConversion.h"
 #include "USDGroomConversion.h"
 #include "USDLayerUtils.h"
@@ -416,6 +417,8 @@ namespace UE::InterchangeUsdTranslator::Private
 			UInterchangeBaseNodeContainer& NodeContainer
 		);
 
+		void DisplayUsdLogMessages(const TArray<TSharedRef<FTokenizedMessage>>& Messages) const;
+
 	public:
 		// We have to keep a stage reference so that we can parse the payloads after Translate() complete.
 		// ReleaseSource() clears this member, once translation is complete.
@@ -465,6 +468,8 @@ namespace UE::InterchangeUsdTranslator::Private
 		// We stash the info we collected from each Volume prim path here, as we'll reuse it between translation and retrieving the payloads
 		TMap<FString, TArray<UsdUtils::FVolumePrimInfo>> PrimPathToVolumeInfo;
 
+		TUniquePtr<FScopedUsdMessageLog> ScopedUsdMessageLog;
+
 	private:
 		TMap<FString, TArray<FUsdMaterialXShaderGraph::FGeomProp>> MaterialUidToGeomProps;
 
@@ -3327,6 +3332,39 @@ namespace UE::InterchangeUsdTranslator::Private
 		}
 	}
 
+	void UInterchangeUSDTranslatorImpl::DisplayUsdLogMessages(const TArray<TSharedRef<FTokenizedMessage>>& Messages) const
+	{
+		for (const TSharedRef<FTokenizedMessage>& Message : Messages)
+		{
+			switch (Message->GetSeverity())
+			{
+				case EMessageSeverity::Error:
+				{
+					UInterchangeResultError_Generic* Result = ResultsContainer->Add<UInterchangeResultError_Generic>();
+					Result->Text = Message->ToText();
+					Result->SourceAssetName = UsdStage.GetRootLayer().GetRealPath();
+					break;
+				}
+				case EMessageSeverity::Warning:
+				case EMessageSeverity::PerformanceWarning:
+				{
+					UInterchangeResultWarning_Generic* Result = ResultsContainer->Add<UInterchangeResultWarning_Generic>();
+					Result->Text = Message->ToText();
+					Result->SourceAssetName = UsdStage.GetRootLayer().GetRealPath();
+					break;
+				}
+				case EMessageSeverity::Info:
+				default:
+				{
+					UInterchangeResultDisplay_Generic* Result = ResultsContainer->Add<UInterchangeResultDisplay_Generic>();
+					Result->Text = Message->ToText();
+					Result->SourceAssetName = UsdStage.GetRootLayer().GetRealPath();
+					break;
+				}
+			}
+		}
+	}
+
 	FString AddLightNode(const UE::FUsdPrim& Prim, UInterchangeBaseNodeContainer& NodeContainer)
 	{
 		TRACE_CPUPROFILER_EVENT_SCOPE(AddLightNode)
@@ -3543,13 +3581,11 @@ namespace UE::InterchangeUsdTranslator::Private
 					FileAttribute.Get(RawFilePath);
 				}
 
-				UE_LOG(
-					LogUsd,
-					Warning,
-					TEXT("Failed to find the file '%s' specified by DomeLight prim '%s'"),
-					*RawFilePath,
-					*Prim.GetPrimPath().GetString()
-				);
+				USD_LOG_USERWARNING(FText::Format(
+					LOCTEXT("MissingLightTexture", "Failed to find the file '{0}' specified by DomeLight prim '{1}'"),
+					FText::FromString(RawFilePath),
+					FText::FromString(Prim.GetPrimPath().GetString())
+				));
 			}
 		}
 
@@ -5278,27 +5314,23 @@ namespace UE::InterchangeUsdTranslator::Private
 			}
 			else
 			{
-				UInterchangeResultWarning_Generic* Message = TranslatorImpl.ResultsContainer->Add<UInterchangeResultWarning_Generic>();
-				Message->SourceAssetName = TranslatorImpl.UsdStage.GetRootLayer().GetRealPath();
-				Message->Text = FText::Format(
+				USD_LOG_USERWARNING(FText::Format(
 					LOCTEXT("UnsupportedSchemaType", "Prim '{0}' has schema '{1}', which is not yet supported via USD Interchange."),
 					FText::FromString(SceneNodeUid),
 					FText::FromName(TypeName)
-				);
+				));
 			}
 		}
 
 		if (Prim.IsA(TEXT("Gprim")) && Prim.GetParent().IsA(TEXT("Gprim")))
 		{
-			UInterchangeResultWarning_Generic* Message = TranslatorImpl.ResultsContainer->Add<UInterchangeResultWarning_Generic>();
-			Message->SourceAssetName = TranslatorImpl.UsdStage.GetRootLayer().GetRealPath();
-			Message->Text = FText::Format(
+			USD_LOG_USERWARNING(FText::Format(
 				LOCTEXT(
 					"NestedGprims",
 					"Prim '{0}' is a Gprim child of another Gprim. This configuration is not recommended or supported very well, and may lead to unexpected results."
 				),
 				FText::FromString(Prim.GetPrimPath().GetString())
-			);
+			));
 		}
 
 		// Only prims that require rendering (and have a renderable parent) get a scene node.
@@ -6623,6 +6655,13 @@ bool UInterchangeUSDTranslator::Translate(UInterchangeBaseNodeContainer& NodeCon
 		return false;
 	}
 
+	ImplPtr->ScopedUsdMessageLog = MakeUnique<FScopedUsdMessageLog>(
+		[ImplPtr](const TArray<TSharedRef<FTokenizedMessage>>& Messages)
+		{
+			ImplPtr->DisplayUsdLogMessages(Messages);
+		}
+	);
+
 	// Setup context
 	UObject* ContextObject = SourceData->GetContextObjectByTag(UE::Interchange::USD::USDContextTag);
 	UInterchangeUsdContext* Context = Cast<UInterchangeUsdContext>(ContextObject);
@@ -6671,12 +6710,10 @@ bool UInterchangeUSDTranslator::Translate(UInterchangeBaseNodeContainer& NodeCon
 
 		if (!StageToImport)
 		{
-			UInterchangeResultError_Generic* Message = ImplPtr->ResultsContainer->Add<UInterchangeResultError_Generic>();
-			Message->SourceAssetName = SourceData->GetFilename();
-			Message->Text = FText::Format(
+			USD_LOG_USERERROR(FText::Format(
 				LOCTEXT("FailedToOpenStage", "Failed to get a valid stage from file '{0}'. Does the file exist and is it a valid USD file?"),
 				FText::FromString(SourceData->GetFilename())
-			);
+			));
 
 			return false;
 		}
@@ -6743,12 +6780,26 @@ bool UInterchangeUSDTranslator::Translate(UInterchangeBaseNodeContainer& NodeCon
 		Traverse(ImplPtr->UsdStage.GetPseudoRoot(), *ImplPtr, NodeContainer, TranslatorSettings, Info);
 	}
 
+	
+
 	return true;
 #else
 	return false;
 #endif	  // USE_USD_SDK
 }
 
+void UInterchangeUSDTranslator::ReleaseSource()
+{
+#if USE_USD_SDK
+	// We need to do this here because by the time ImportFinish is called it's too late to add new messages to the results container
+	UE::InterchangeUsdTranslator::Private::UInterchangeUSDTranslatorImpl* ImplPtr = Impl.Get();
+	if (ImplPtr)
+	{
+		ImplPtr->ScopedUsdMessageLog.Reset();
+	}
+#endif	  // USE_USD_SDK
+}
+
 void UInterchangeUSDTranslator::ImportFinish()
 {
 	// These objects are only set by Translate so they are released here instead of ReleaseSource
diff --git a/Engine/Plugins/Interchange/Extensions/OpenUSD/Source/Import/Public/InterchangeUsdTranslator.h b/Engine/Plugins/Interchange/Extensions/OpenUSD/Source/Import/Public/InterchangeUsdTranslator.h
--- a/Engine/Plugins/Interchange/Extensions/OpenUSD/Source/Import/Public/InterchangeUsdTranslator.h
+++ b/Engine/Plugins/Interchange/Extensions/OpenUSD/Source/Import/Public/InterchangeUsdTranslator.h
@@ -108,6 +108,7 @@ public:
 	virtual EInterchangeTranslatorAssetType GetSupportedAssetTypes() const override;
 	virtual TArray<FString> GetSupportedFormats() const override;
 	virtual bool Translate(UInterchangeBaseNodeContainer& BaseNodeContainer) const override;
+	virtual void ReleaseSource() override;
 	virtual void ImportFinish() override;
 	virtual UInterchangeTranslatorSettings* GetSettings() const override;
 	virtual void SetSettings(const UInterchangeTranslatorSettings* InterchangeTranslatorSettings) override;
diff --git a/Engine/Plugins/Runtime/USDCore/Source/USDUtilities/Private/USDErrorUtils.cpp b/Engine/Plugins/Runtime/USDCore/Source/USDUtilities/Private/USDErrorUtils.cpp
--- a/Engine/Plugins/Runtime/USDCore/Source/USDUtilities/Private/USDErrorUtils.cpp
+++ b/Engine/Plugins/Runtime/USDCore/Source/USDUtilities/Private/USDErrorUtils.cpp
@@ -189,7 +189,8 @@ namespace UE::USDErrorUtils::Private
 		FUsdCombinedMessages& operator=(const FUsdCombinedMessages& Other) = delete;
 		FUsdCombinedMessages& operator=(FUsdCombinedMessages&& Other) = default;
 
-		FUsdCombinedMessages()
+		FUsdCombinedMessages(FUsdLogManager::FMessageDisplayCallback InCallback)
+			: MessageDisplayCallback(InCallback)
 		{
 			const UUsdProjectSettings* ProjectSettings = GetDefault<UUsdProjectSettings>();
 			bOptimizeUsdLog = ProjectSettings->bOptimizeUsdLog;
@@ -258,23 +259,9 @@ namespace UE::USDErrorUtils::Private
 				}
 			}
 
-			if (TokenizedMessages.Num() > 0)
+			if (MessageDisplayCallback)
 			{
-				FMessageLogModule& MessageLogModule = FModuleManager::LoadModuleChecked<FMessageLogModule>("MessageLog");
-				TSharedRef<IMessageLogListing> LogListing = MessageLogModule.GetLogListing(TEXT("USD"));
-
-				// We'll output the messages to the output log ourselves from LogMessageInternal()
-				const bool bMirrorToOutputLog = false;
-				LogListing->AddMessages(TokenizedMessages, bMirrorToOutputLog);
-
-				// Force display even if all we have are info level messages.
-				// This is also handy because it also outputs them on the output log for us
-				const bool bForce = true;
-				LogListing->NotifyIfAnyMessages(	//
-					LOCTEXT("Log", "The previous USD operation produced some log messages."),
-					EMessageSeverity::Info,
-					bForce
-				);
+				MessageDisplayCallback(TokenizedMessages);
 			}
 #endif	  // WITH_EDITOR
 		}
@@ -344,6 +331,7 @@ namespace UE::USDErrorUtils::Private
 		bool bHasSkippedMessages = false;
 		TMap<FName, FMessageInfo> LoggedMessagesByID;
 		TArray<FMessageInfo> UncollapsedLoggedMessages;
+		FUsdLogManager::FMessageDisplayCallback MessageDisplayCallback;
 	};
 
 	void LogMessageInternal(const TSharedRef<FTokenizedMessage>& Message, bool bUserFacing)
@@ -447,7 +435,33 @@ void FUsdLogManager::UnregisterDiagnosticDelegate()
 #endif	  // USE_USD_SDK
 }
 
-FScopedUsdMessageLog::FScopedUsdMessageLog()
+void FUsdLogManager::DisplayMessagesOnMessageLog(const TArray<TSharedRef<FTokenizedMessage>>& Messages)
+{
+#if WITH_EDITOR
+	if (Messages.Num() == 0)
+	{
+		return;
+	}
+
+	FMessageLogModule& MessageLogModule = FModuleManager::LoadModuleChecked<FMessageLogModule>("MessageLog");
+	TSharedRef<IMessageLogListing> LogListing = MessageLogModule.GetLogListing(TEXT("USD"));
+
+	// We'll output the messages to the output log ourselves from LogMessageInternal()
+	const bool bMirrorToOutputLog = false;
+	LogListing->AddMessages(Messages, bMirrorToOutputLog);
+
+	// Force display even if all we have are info level messages.
+	// This is also handy because it also outputs them on the output log for us
+	const bool bForce = true;
+	LogListing->NotifyIfAnyMessages(	//
+		LOCTEXT("Log", "The previous USD operation produced some log messages."),
+		EMessageSeverity::Info,
+		bForce
+	);
+#endif	  // USE_USD_SDK
+}
+
+FScopedUsdMessageLog::FScopedUsdMessageLog(FUsdLogManager::FMessageDisplayCallback InCallback)
 {
 #if USE_USD_SDK
 	using namespace UE::USDErrorUtils::Private;
@@ -456,7 +470,7 @@ FScopedUsdMessageLog::FScopedUsdMessageLog()
 
 	if (++CombinedMessagesRefCount == 1)
 	{
-		CombinedMessages = MakeUnique<FUsdCombinedMessages>();
+		CombinedMessages = MakeUnique<FUsdCombinedMessages>(InCallback);
 		PushErrorMark();
 	}
 
diff --git a/Engine/Plugins/Runtime/USDCore/Source/USDUtilities/Public/USDErrorUtils.h b/Engine/Plugins/Runtime/USDCore/Source/USDUtilities/Public/USDErrorUtils.h
--- a/Engine/Plugins/Runtime/USDCore/Source/USDUtilities/Public/USDErrorUtils.h
+++ b/Engine/Plugins/Runtime/USDCore/Source/USDUtilities/Public/USDErrorUtils.h
@@ -36,6 +36,8 @@ namespace UsdUtils
 class FUsdLogManager
 {
 public:
+	using FMessageDisplayCallback = TFunction<void(const TArray<TSharedRef<FTokenizedMessage>>&)>;
+
 	// Copied from AssertionMacros.h as it's private there
 	UE_NODEBUG static constexpr uint32 FileHashForEnsure(const char* Filename)
 	{
@@ -67,6 +69,8 @@ public:
 
 	static UE_API void UnregisterDiagnosticDelegate();
 
+	static UE_API void DisplayMessagesOnMessageLog(const TArray<TSharedRef<FTokenizedMessage>>& Messages);
+
 	UE_DEPRECATED(5.6, "Use Log, or ideally the USD_LOG_INFO, USD_LOG_WARNING, etc. macros")
 	static UE_API void LogMessage(EMessageSeverity::Type Severity, const FText& Message);
 
@@ -111,7 +115,7 @@ public:
 class FScopedUsdMessageLog
 {
 public:
-	UE_API FScopedUsdMessageLog();
+	UE_API FScopedUsdMessageLog(FUsdLogManager::FMessageDisplayCallback = FUsdLogManager::DisplayMessagesOnMessageLog);
 	UE_API ~FScopedUsdMessageLog();
 
 	FScopedUsdMessageLog(const FScopedUsdMessageLog&) = delete;
