USD Interchange: Fix light color being left at an uninitialized value if there is no USD opinion

diff --git a/Engine/Plugins/Interchange/Extensions/OpenUSD/Source/Import/Private/InterchangeUsdTranslator.cpp b/Engine/Plugins/Interchange/Extensions/OpenUSD/Source/Import/Private/InterchangeUsdTranslator.cpp
--- a/Engine/Plugins/Interchange/Extensions/OpenUSD/Source/Import/Private/InterchangeUsdTranslator.cpp
+++ b/Engine/Plugins/Interchange/Extensions/OpenUSD/Source/Import/Private/InterchangeUsdTranslator.cpp
@@ -3384,9 +3384,9 @@ namespace UE::InterchangeUsdTranslator::Private
 
 		// Temporarily a bit awkward here as we updated our VtValue::Get() to not perform any type coercion, so we
 		// need to do that ourselves. In the future we'll have a dedicated templated function to perform this process
-		FLinearColor Color;
+		FLinearColor Color = FLinearColor::White;
 		UE::FUsdAttribute ColorAttribute = Prim.GetAttribute(*ColorToken);
-		if (ColorAttribute && ColorAttribute.HasAuthoredValue())
+		if (ColorAttribute)
 		{
 			UE::FVtValue Value;
 			if (ColorAttribute.Get(Value))
