USD Interchange: Only author translated node attributes for USD attributes with authored opinions

diff --git a/Engine/Plugins/Interchange/Extensions/OpenUSD/Source/Import/Private/SchemaHandlers/CameraSchemaHandler.cpp b/Engine/Plugins/Interchange/Extensions/OpenUSD/Source/Import/Private/SchemaHandlers/CameraSchemaHandler.cpp
--- a/Engine/Plugins/Interchange/Extensions/OpenUSD/Source/Import/Private/SchemaHandlers/CameraSchemaHandler.cpp
+++ b/Engine/Plugins/Interchange/Extensions/OpenUSD/Source/Import/Private/SchemaHandlers/CameraSchemaHandler.cpp
@@ -37,17 +37,23 @@ namespace UE::GeomCameraSchemaHandler::Private
 		static const FString HorizontalApertureToken = UsdToUnreal::ConvertToken(pxr::UsdGeomTokens->horizontalAperture);
 		static const FString VerticalApertureToken = UsdToUnreal::ConvertToken(pxr::UsdGeomTokens->verticalAperture);
 
-		float FocalLength = UsdUtils::GetAttributeValue<float>(Prim, FocalLengthToken);
-		FocalLength = UsdToUnreal::ConvertDistance(StageInfo, FocalLength);
-		CameraNode->SetCustomFocalLength(FocalLength);
+		if (float FocalLength; UsdUtils::GetAuthoredAttributeValue<float>(Prim, FocalLengthToken, FocalLength))
+		{
+			FocalLength = UsdToUnreal::ConvertDistance(StageInfo, FocalLength);
+			CameraNode->SetCustomFocalLength(FocalLength);
+		}
 
-		float SensorWidth = UsdUtils::GetAttributeValue<float>(Prim, HorizontalApertureToken);
-		SensorWidth = UsdToUnreal::ConvertDistance(StageInfo, SensorWidth);
-		CameraNode->SetCustomSensorWidth(SensorWidth);
+		if (float SensorWidth; UsdUtils::GetAuthoredAttributeValue<float>(Prim, HorizontalApertureToken, SensorWidth))
+		{
+			SensorWidth = UsdToUnreal::ConvertDistance(StageInfo, SensorWidth);
+			CameraNode->SetCustomSensorWidth(SensorWidth);
+		}
 
-		float SensorHeight = UsdUtils::GetAttributeValue<float>(Prim, VerticalApertureToken);
-		SensorHeight = UsdToUnreal::ConvertDistance(StageInfo, SensorHeight);
-		CameraNode->SetCustomSensorHeight(SensorHeight);
+		if (float SensorHeight; UsdUtils::GetAuthoredAttributeValue<float>(Prim, VerticalApertureToken, SensorHeight))
+		{
+			SensorHeight = UsdToUnreal::ConvertDistance(StageInfo, SensorHeight);
+			CameraNode->SetCustomSensorHeight(SensorHeight);
+		}
 
 		// Focus distance and FStop not currently supported
 	}
diff --git a/Engine/Plugins/Interchange/Extensions/OpenUSD/Source/Import/Private/SchemaHandlers/LightSchemaHandler.cpp b/Engine/Plugins/Interchange/Extensions/OpenUSD/Source/Import/Private/SchemaHandlers/LightSchemaHandler.cpp
--- a/Engine/Plugins/Interchange/Extensions/OpenUSD/Source/Import/Private/SchemaHandlers/LightSchemaHandler.cpp
+++ b/Engine/Plugins/Interchange/Extensions/OpenUSD/Source/Import/Private/SchemaHandlers/LightSchemaHandler.cpp
@@ -108,21 +108,21 @@ namespace UE::LuxLightSchemaHandler::Private
 
 		// Ref. UsdToUnreal::ConvertLight
 		static const FString IntensityToken = UsdToUnreal::ConvertToken(pxr::UsdLuxTokens->inputsIntensity);
-		float Intensity = UsdUtils::GetAttributeValue<float>(Prim, IntensityToken);
+		float Intensity = 1.0f;
+		bool bIntensityAuthored = UsdUtils::GetAuthoredAttributeValue<float>(Prim, IntensityToken, Intensity);
 
 		static const FString ExposureToken = UsdToUnreal::ConvertToken(pxr::UsdLuxTokens->inputsExposure);
-		float Exposure = UsdUtils::GetAttributeValue<float>(Prim, ExposureToken);
+		float Exposure = 0.0f;
+		UsdUtils::GetAuthoredAttributeValue<float>(Prim, ExposureToken, Exposure);
 
 		static const FString RadiusToken = UsdToUnreal::ConvertToken(pxr::UsdLuxTokens->inputsRadius);
 
-		// TODO: Maybe only author these in case there were in fact authored opinions on the prims? Otherwise we're just authoring default values
-
 		// Temporarily a bit awkward here as we updated our VtValue::Get() to not perform any type coercion, so we
 		// need to do that ourselves. In the future we'll have a dedicated templated function to perform this process
 		FLinearColor Color = FLinearColor::White;
 		static const FString ColorToken = UsdToUnreal::ConvertToken(pxr::UsdLuxTokens->inputsColor);
 		UE::FUsdAttribute ColorAttribute = Prim.GetAttribute(*ColorToken);
-		if (ColorAttribute)
+		if (ColorAttribute && ColorAttribute.HasAuthoredValue())
 		{
 			UE::FVtValue Value;
 			if (ColorAttribute.Get(Value))
@@ -153,18 +153,23 @@ namespace UE::LuxLightSchemaHandler::Private
 					Color = FLinearColor{Value.Get<FVector4d>()};
 				}
 			}
+
+			const bool bSRGB = true;
+			Color = Color.ToFColor(bSRGB);
+			LightNode->SetCustomLightColor(Color);
 		}
-		const bool bSRGB = true;
-		Color = Color.ToFColor(bSRGB);
-		LightNode->SetCustomLightColor(Color);
 
 		static const FString TemperatureToken = UsdToUnreal::ConvertToken(pxr::UsdLuxTokens->inputsColorTemperature);
-		float Temperature = UsdUtils::GetAttributeValue<float>(Prim, TemperatureToken);
-		LightNode->SetCustomTemperature(Temperature);
+		if (float Temperature; UsdUtils::GetAuthoredAttributeValue<float>(Prim, TemperatureToken, Temperature))
+		{
+			LightNode->SetCustomTemperature(Temperature);
+		}
 
 		static const FString UseTemperatureToken = UsdToUnreal::ConvertToken(pxr::UsdLuxTokens->inputsEnableColorTemperature);
-		bool UseTemperature = UsdUtils::GetAttributeValue<bool>(Prim, UseTemperatureToken);
-		LightNode->SetCustomUseTemperature(UseTemperature);
+		if (bool UseTemperature; UsdUtils::GetAuthoredAttributeValue<bool>(Prim, UseTemperatureToken, UseTemperature))
+		{
+			LightNode->SetCustomUseTemperature(UseTemperature);
+		}
 
 		// "Shadow enabled" currently not supported
 
@@ -176,19 +181,24 @@ namespace UE::LuxLightSchemaHandler::Private
 				return;
 			}
 
-			Intensity = UsdToUnreal::ConvertLightIntensityAttr(Intensity, Exposure);
-			LightNode->SetCustomIntensity(Intensity);
+			if (bIntensityAuthored)
+			{
+				Intensity = UsdToUnreal::ConvertLightIntensityAttr(Intensity, Exposure);
+				LightNode->SetCustomIntensity(Intensity);
+			}
 
 			static const FString SourceAngleToken = UsdToUnreal::ConvertToken(pxr::UsdLuxTokens->inputsAngle);
-			float Angle = UsdUtils::GetAttributeValue<float>(Prim, SourceAngleToken);
-			CastLightNode->SetCustomLightSourceAngle(Angle);
+			if (float Angle; UsdUtils::GetAuthoredAttributeValue<float>(Prim, SourceAngleToken, Angle))
+			{
+				CastLightNode->SetCustomLightSourceAngle(Angle);
+			}
 		}
 		else if (Prim.IsA(TEXT("SphereLight")))
 		{
 			const FUsdStageInfo StageInfo(Prim.GetStage());
 
-			const float Radius = UsdUtils::GetAttributeValue<float>(Prim, RadiusToken);
-			const float SourceRadius = UsdToUnreal::ConvertDistance(StageInfo, Radius);	   // currently not supported
+			float Radius = 0.5f;
+			UsdUtils::GetAuthoredAttributeValue<float>(Prim, RadiusToken, Radius);
 
 			if (Prim.HasAPI(TEXT("ShapingAPI")))
 			{
@@ -198,22 +208,29 @@ namespace UE::LuxLightSchemaHandler::Private
 					return;
 				}
 
-				CastLightNode->SetCustomIntensityUnits(EInterchangeLightUnits::Lumens);
-
 				static const FString ConeAngleToken = UsdToUnreal::ConvertToken(pxr::UsdLuxTokens->inputsShapingConeAngle);
 				static const FString ConeSoftnessToken = UsdToUnreal::ConvertToken(pxr::UsdLuxTokens->inputsShapingConeSoftness);
 
-				float ConeAngle = UsdUtils::GetAttributeValue<float>(Prim, ConeAngleToken);
-				float ConeSoftness = UsdUtils::GetAttributeValue<float>(Prim, ConeSoftnessToken);
+				float ConeAngle = 90.0f;
+				bool bConeAngleAuthored = UsdUtils::GetAuthoredAttributeValue<float>(Prim, ConeAngleToken, ConeAngle);
 
-				float InnerConeAngle = 0.0f;
-				const float OuterConeAngle = UsdToUnreal::ConvertConeAngleSoftnessAttr(ConeAngle, ConeSoftness, InnerConeAngle);
+				float ConeSoftness = 0.0f;
+				bool bConeSoftnessAuthored = UsdUtils::GetAuthoredAttributeValue<float>(Prim, ConeSoftnessToken, ConeSoftness);
 
-				Intensity = UsdToUnreal::ConvertLuxShapingAPIIntensityAttr(Intensity, Exposure, Radius, ConeAngle, ConeSoftness, StageInfo);
-				LightNode->SetCustomIntensity(Intensity);
+				if (bConeAngleAuthored || bConeSoftnessAuthored)
+				{
+					float InnerConeAngle = 0.0f;
+					const float OuterConeAngle = UsdToUnreal::ConvertConeAngleSoftnessAttr(ConeAngle, ConeSoftness, InnerConeAngle);
+					CastLightNode->SetCustomInnerConeAngle(InnerConeAngle);
+					CastLightNode->SetCustomOuterConeAngle(OuterConeAngle);
+				}
 
-				CastLightNode->SetCustomInnerConeAngle(InnerConeAngle);
-				CastLightNode->SetCustomOuterConeAngle(OuterConeAngle);
+				if (bIntensityAuthored)
+				{
+					Intensity = UsdToUnreal::ConvertLuxShapingAPIIntensityAttr(Intensity, Exposure, Radius, ConeAngle, ConeSoftness, StageInfo);
+					CastLightNode->SetCustomIntensityUnits(EInterchangeLightUnits::Lumens);
+					CastLightNode->SetCustomIntensity(Intensity);
+				}
 			}
 			else
 			{
@@ -223,10 +240,12 @@ namespace UE::LuxLightSchemaHandler::Private
 					return;
 				}
 
-				CastLightNode->SetCustomIntensityUnits(EInterchangeLightUnits::Lumens);
-
-				Intensity = UsdToUnreal::ConvertSphereLightIntensityAttr(Intensity, Exposure, Radius, StageInfo);
-				CastLightNode->SetCustomIntensity(Intensity);
+				if (bIntensityAuthored)
+				{
+					Intensity = UsdToUnreal::ConvertSphereLightIntensityAttr(Intensity, Exposure, Radius, StageInfo);
+					CastLightNode->SetCustomIntensityUnits(EInterchangeLightUnits::Lumens);
+					CastLightNode->SetCustomIntensity(Intensity);
+				}
 			}
 		}
 		else if (Prim.IsA(TEXT("RectLight")) || Prim.IsA(TEXT("DiskLight")))
@@ -237,33 +256,52 @@ namespace UE::LuxLightSchemaHandler::Private
 				return;
 			}
 
-			CastLightNode->SetCustomIntensityUnits(EInterchangeLightUnits::Lumens);
-
 			static const FString WidthToken = UsdToUnreal::ConvertToken(pxr::UsdLuxTokens->inputsWidth);
 			static const FString HeightToken = UsdToUnreal::ConvertToken(pxr::UsdLuxTokens->inputsHeight);
 
-			float Width = UsdUtils::GetAttributeValue<float>(Prim, WidthToken);
-			float Height = UsdUtils::GetAttributeValue<float>(Prim, HeightToken);
-
 			const FUsdStageInfo StageInfo(Prim.GetStage());
 
 			if (Prim.IsA(TEXT("RectLight")))
 			{
-				Width = UsdToUnreal::ConvertDistance(StageInfo, Width);
-				Height = UsdToUnreal::ConvertDistance(StageInfo, Height);
-				Intensity = UsdToUnreal::ConvertRectLightIntensityAttr(Intensity, Exposure, Width, Height, StageInfo);
+				float Width = 1.0f;
+				if (UsdUtils::GetAuthoredAttributeValue<float>(Prim, WidthToken, Width))
+				{
+					Width = UsdToUnreal::ConvertDistance(StageInfo, Width);
+					CastLightNode->SetCustomSourceWidth(Width);
+				}
+
+				float Height = 1.0f;
+				if (UsdUtils::GetAuthoredAttributeValue<float>(Prim, HeightToken, Height))
+				{
+					Height = UsdToUnreal::ConvertDistance(StageInfo, Height);
+					CastLightNode->SetCustomSourceHeight(Height);
+				}
+
+				if (bIntensityAuthored)
+				{
+					Intensity = UsdToUnreal::ConvertRectLightIntensityAttr(Intensity, Exposure, Width, Height, StageInfo);
+					CastLightNode->SetCustomIntensityUnits(EInterchangeLightUnits::Lumens);
+					CastLightNode->SetCustomIntensity(Intensity);
+				}
 			}
-			else
+			else // DiskLight
 			{
-				float Radius = UsdUtils::GetAttributeValue<float>(Prim, RadiusToken);
-				Width = UsdToUnreal::ConvertDistance(StageInfo, Radius) * 2.f;
-				Height = Width;
+				float Radius = 0.5f;
+				if (UsdUtils::GetAuthoredAttributeValue<float>(Prim, RadiusToken, Radius))
+				{
+					float Width = UsdToUnreal::ConvertDistance(StageInfo, Radius) * 2.f;
+					float Height = Width;
+					CastLightNode->SetCustomSourceWidth(Width);
+					CastLightNode->SetCustomSourceHeight(Height);
+				}
 
-				Intensity = UsdToUnreal::ConvertDiskLightIntensityAttr(Intensity, Exposure, Radius, StageInfo);
+				if (bIntensityAuthored)
+				{
+					Intensity = UsdToUnreal::ConvertDiskLightIntensityAttr(Intensity, Exposure, Radius, StageInfo);
+					CastLightNode->SetCustomIntensityUnits(EInterchangeLightUnits::Lumens);
+					CastLightNode->SetCustomIntensity(Intensity);
+				}
 			}
-			CastLightNode->SetCustomIntensity(Intensity);
-			CastLightNode->SetCustomSourceWidth(Width);
-			CastLightNode->SetCustomSourceHeight(Height);
 		}
 		else if (Prim.IsA(TEXT("DomeLight")))
 		{
@@ -275,8 +313,11 @@ namespace UE::LuxLightSchemaHandler::Private
 
 			// This doesn't translate very well as SkyLightComponents seem to have an "IntensityScale", and there's never any clear
 			// indication of what the dome light intensity actually is... But I don't think there's much we can do
-			Intensity = UsdToUnreal::ConvertLightIntensityAttr(Intensity, Exposure);
-			LightNode->SetCustomIntensity(Intensity);
+			if (bIntensityAuthored)
+			{
+				Intensity = UsdToUnreal::ConvertLightIntensityAttr(Intensity, Exposure);
+				LightNode->SetCustomIntensity(Intensity);
+			}
 
 			// Show some warning in case the file specifies a different DomeLight format, as we only support latlong
 			static const FString FormatToken = UsdToUnreal::ConvertToken(pxr::UsdLuxTokens->inputsTextureFormat);
diff --git a/Engine/Plugins/Runtime/USDCore/Source/UnrealUSDWrapper/Private/UsdWrappers/UsdAttribute.cpp b/Engine/Plugins/Runtime/USDCore/Source/UnrealUSDWrapper/Private/UsdWrappers/UsdAttribute.cpp
--- a/Engine/Plugins/Runtime/USDCore/Source/UnrealUSDWrapper/Private/UsdWrappers/UsdAttribute.cpp
+++ b/Engine/Plugins/Runtime/USDCore/Source/UnrealUSDWrapper/Private/UsdWrappers/UsdAttribute.cpp
@@ -447,11 +447,30 @@ namespace UsdUtils
 
 		return Value;
 	}
+
+	template<typename ValueType>
+	bool GetAuthoredAttributeValue(const UE::FUsdPrim& InPrim, FStringView InAttributeName, ValueType& OutValue, TOptional<double> InTime)
+	{
+		if (!InPrim)
+		{
+			return false;
+		}
+
+		UE::FUsdAttribute Attribute = InPrim.GetAttribute(InAttributeName.GetData());
+		if (!Attribute)
+		{
+			return false;
+		}
+
+		const bool bGotValue = Attribute.Get(OutValue, InTime);
+		return bGotValue && Attribute.HasAuthoredValue();
+	}
 }	 // namespace UsdUtils
 
 #define SPECIALIZE(X) \
 	template UNREALUSDWRAPPER_API bool UE::FUsdAttribute::Get(X& Value, TOptional<double> Time) const;\
-	template UNREALUSDWRAPPER_API X UsdUtils::GetAttributeValue(const UE::FUsdPrim& Prim, FStringView AttributeName, TOptional<double> Time);
+	template UNREALUSDWRAPPER_API X UsdUtils::GetAttributeValue(const UE::FUsdPrim& Prim, FStringView AttributeName, TOptional<double> Time); \
+	template UNREALUSDWRAPPER_API bool UsdUtils::GetAuthoredAttributeValue(const UE::FUsdPrim& InPrim, FStringView InAttributeName, X& OutValue, TOptional<double> InTime);
 
 SPECIALIZE(bool);
 SPECIALIZE(uint8);
diff --git a/Engine/Plugins/Runtime/USDCore/Source/UnrealUSDWrapper/Public/UsdWrappers/UsdAttribute.h b/Engine/Plugins/Runtime/USDCore/Source/UnrealUSDWrapper/Public/UsdWrappers/UsdAttribute.h
--- a/Engine/Plugins/Runtime/USDCore/Source/UnrealUSDWrapper/Public/UsdWrappers/UsdAttribute.h
+++ b/Engine/Plugins/Runtime/USDCore/Source/UnrealUSDWrapper/Public/UsdWrappers/UsdAttribute.h
@@ -119,4 +119,10 @@ namespace UsdUtils
 	// Returns the attribute value for the given prim and attribute name. If the attribute doesn't exist, it returns the default value for the type.
 	template<typename ValueType>
 	UNREALUSDWRAPPER_API ValueType GetAttributeValue(const UE::FUsdPrim& Prim, FStringView AttributeName, TOptional<double> Time = {});
+
+	// Outputs the authored attribute value for the given prim and attribute name, and returns true.
+	// If the attribute doesn't exist, returns false and leaves OutValue unassigned.
+	// If the attribute does exist but has no authored opinion, returns false but assigns OutValue with the attribute's fallback value, if possible.
+	template<typename ValueType>
+	UNREALUSDWRAPPER_API bool GetAuthoredAttributeValue(const UE::FUsdPrim& InPrim, FStringView InAttributeName, ValueType& OutValue, TOptional<double> InTime = {});
 }
\ No newline at end of file
