Interchange USD: Fixed potential crash when importing a geometry cache with multiple sections with the same material binding.

#rnx
#Synced-CL 48607582

diff --git a/Engine/Plugins/Interchange/Extensions/OpenUSD/Source/Import/Private/SchemaHandlers/SchemaHandlerUtils.cpp b/Engine/Plugins/Interchange/Extensions/OpenUSD/Source/Import/Private/SchemaHandlers/SchemaHandlerUtils.cpp
--- a/Engine/Plugins/Interchange/Extensions/OpenUSD/Source/Import/Private/SchemaHandlers/SchemaHandlerUtils.cpp
+++ b/Engine/Plugins/Interchange/Extensions/OpenUSD/Source/Import/Private/SchemaHandlers/SchemaHandlerUtils.cpp
@@ -10,6 +10,7 @@
 #include "Usd/InterchangeUsdDefinitions.h"
 #include "USDErrorUtils.h"
 #include "USDNaniteAssemblyUtils.h"
+#include "USDObjectUtils.h"
 #include "USDPrimConversion.h"
 #include "USDShadeConversion.h"
 #include "USDTypesConversion.h"
@@ -1125,6 +1126,8 @@ namespace UE::Interchange::USD
 	{
 		using namespace USDSchemaHandlerUtils::Private;
 
+		TSet<FString> UsedSlotNames;
+
 		for (const UsdUtils::FUsdPrimMaterialSlot& Slot : MaterialAssignments.Slots)
 		{
 			// We do this because Interchange will, in some scenarios, merge material slots with identical slot names.
@@ -1179,6 +1182,12 @@ namespace UE::Interchange::USD
 				case UsdUtils::EPrimAssignmentType::UnrealMaterial:
 				{
 					// MaterialSource here is the content path, e.g. "/Game/MyFolder/Red.Red"
+					// Use the object name as the slot name instead of using the full path
+					int32 SplitIndex = INDEX_NONE;
+					if (SlotName.FindChar('.', SplitIndex))
+					{
+						SlotName = SlotName.Right(SlotName.Len() - (SplitIndex + 1)); // + 1 to keep just the part after the period
+					}
 					MaterialUid = AddUnrealMaterialReferenceNodeIfNeeded(Slot.MaterialSource, *UsdContext.GetNodeContainer(), AccumulatedInfo);
 					break;
 				}
@@ -1189,7 +1198,12 @@ namespace UE::Interchange::USD
 				}
 			}
 
-			MeshOrSceneNode->SetSlotMaterialDependencyUid(SlotName, *MaterialUid);
+			// It's possible to have multiple mesh sections with the same material bindings within a mesh node so make sure the
+			// slot names are unique to keep them separate
+			FString UniqueSlotName = UsdUnreal::ObjectUtils::GetUniqueName(SlotName, UsedSlotNames);
+			UsedSlotNames.Add(UniqueSlotName);
+
+			MeshOrSceneNode->SetSlotMaterialDependencyUid(UniqueSlotName, *MaterialUid);
 		}
 	}
 	template void SetSlotMaterialDependencies(UInterchangeMeshNode*, const UsdUtils::FUsdPrimMaterialAssignmentInfo&, FHandlerAccumulatedInfo&, UInterchangeUsdContext&);
