USD: Partial revert of 48740998 as the slot names placed in the translated nodes must match what we'll eventually place on the mesh description

#Synced-CL 48758811

diff --git a/Engine/Plugins/Interchange/Extensions/OpenUSD/Source/Import/Private/SchemaHandlers/SchemaHandlerUtils.cpp b/Engine/Plugins/Interchange/Extensions/OpenUSD/Source/Import/Private/SchemaHandlers/SchemaHandlerUtils.cpp
--- a/Engine/Plugins/Interchange/Extensions/OpenUSD/Source/Import/Private/SchemaHandlers/SchemaHandlerUtils.cpp
+++ b/Engine/Plugins/Interchange/Extensions/OpenUSD/Source/Import/Private/SchemaHandlers/SchemaHandlerUtils.cpp
@@ -1181,12 +1181,6 @@ namespace UE::Interchange::USD
 				case UsdUtils::EPrimAssignmentType::UnrealMaterial:
 				{
 					// MaterialSource here is the content path, e.g. "/Game/MyFolder/Red.Red"
-					// Use the object name as the slot name instead of using the full path
-					int32 SplitIndex = INDEX_NONE;
-					if (SlotName.FindChar('.', SplitIndex))
-					{
-						SlotName = SlotName.Right(SlotName.Len() - (SplitIndex + 1)); // + 1 to keep just the part after the period
-					}
 					MaterialUid = AddUnrealMaterialReferenceNodeIfNeeded(Slot.MaterialSource, *UsdContext.GetNodeContainer(), AccumulatedInfo);
 					break;
 				}
