USD: Fix needed missing includes after 48985143

#Synced-CL 48758811

diff --git a/Engine/Plugins/Interchange/Extensions/OpenUSD/Source/Import/Private/SchemaHandlers/UnrealMaterialSchemaHandler.cpp b/Engine/Plugins/Interchange/Extensions/OpenUSD/Source/Import/Private/SchemaHandlers/UnrealMaterialSchemaHandler.cpp
--- a/Engine/Plugins/Interchange/Extensions/OpenUSD/Source/Import/Private/SchemaHandlers/UnrealMaterialSchemaHandler.cpp
+++ b/Engine/Plugins/Interchange/Extensions/OpenUSD/Source/Import/Private/SchemaHandlers/UnrealMaterialSchemaHandler.cpp
@@ -10,6 +10,7 @@
 #include "InterchangeUsdTraversalInfo.h"
 #include "SchemaHandlers/SchemaHandlerUtils.h"
 #include "Usd/InterchangeUsdDefinitions.h"
+#include "USDConversionUtils.h"
 #include "USDShadeConversion.h"
 #include "UsdWrappers/UsdPrim.h"
 #include "UsdWrappers/UsdStage.h"
diff --git a/Engine/Plugins/Interchange/Extensions/OpenUSD/Source/Import/Private/SchemaHandlers/XformableSchemaHandler.cpp b/Engine/Plugins/Interchange/Extensions/OpenUSD/Source/Import/Private/SchemaHandlers/XformableSchemaHandler.cpp
--- a/Engine/Plugins/Interchange/Extensions/OpenUSD/Source/Import/Private/SchemaHandlers/XformableSchemaHandler.cpp
+++ b/Engine/Plugins/Interchange/Extensions/OpenUSD/Source/Import/Private/SchemaHandlers/XformableSchemaHandler.cpp
@@ -10,6 +10,7 @@
 #include "SchemaHandlers/SchemaHandlerUtils.h"
 #include "USDPrimConversion.h"
 #include "UsdWrappers/UsdPrim.h"
+#include "UsdWrappers/UsdTyped.h"
 
 #include "InterchangeSceneComponentNodes.h"
 #include "InterchangeSceneNode.h"
