USD: Remove remaining "b" prefix from new property that was changed from bool to an enum

#rnx
#Synced-CL 48758811

diff --git a/Engine/Plugins/Interchange/Extensions/OpenUSD/Source/Import/Private/InterchangeUsdTranslator.cpp b/Engine/Plugins/Interchange/Extensions/OpenUSD/Source/Import/Private/InterchangeUsdTranslator.cpp
--- a/Engine/Plugins/Interchange/Extensions/OpenUSD/Source/Import/Private/InterchangeUsdTranslator.cpp
+++ b/Engine/Plugins/Interchange/Extensions/OpenUSD/Source/Import/Private/InterchangeUsdTranslator.cpp
@@ -134,7 +134,7 @@ UInterchangeUsdTranslatorSettings::UInterchangeUsdTranslatorSettings()
 		  0.01,				   // MetersPerUnit
 		  EUsdUpAxis::ZAxis	   // UpAxis
 	  }
-	, bPointInstancerCollapsing(EUsdPointInstancerCollapsing::DontCollapse)
+	, PointInstancerCollapsing(EUsdPointInstancerCollapsing::DontCollapse)
 	, bUseSchemaForCollapsing(false)
 	, bUsePrimKindsForCollapsing(false)
 	, KindsToCollapse(static_cast<int32>(EUsdDefaultKind::Component | EUsdDefaultKind::Subcomponent))
diff --git a/Engine/Plugins/Interchange/Extensions/OpenUSD/Source/Import/Private/SchemaHandlers/PointInstancerSchemaHandler.cpp b/Engine/Plugins/Interchange/Extensions/OpenUSD/Source/Import/Private/SchemaHandlers/PointInstancerSchemaHandler.cpp
--- a/Engine/Plugins/Interchange/Extensions/OpenUSD/Source/Import/Private/SchemaHandlers/PointInstancerSchemaHandler.cpp
+++ b/Engine/Plugins/Interchange/Extensions/OpenUSD/Source/Import/Private/SchemaHandlers/PointInstancerSchemaHandler.cpp
@@ -410,7 +410,7 @@ namespace UE::Interchange::USD
 	{
 		// If we can't collapse point instancers then we already know for sure this can't be collapsed
 		UInterchangeUsdTranslatorSettings* Settings = UsdContext.GetTranslatorSettings();
-		if (Settings && Settings->bPointInstancerCollapsing == EUsdPointInstancerCollapsing::DontCollapse)
+		if (Settings && Settings->PointInstancerCollapsing == EUsdPointInstancerCollapsing::DontCollapse)
 		{
 			return false;
 		}
@@ -422,7 +422,7 @@ namespace UE::Interchange::USD
 	{
 		// If we can't collapse point instancers then we already know for sure this can't collapse
 		UInterchangeUsdTranslatorSettings* Settings = UsdContext.GetTranslatorSettings();
-		if (Settings && Settings->bPointInstancerCollapsing == EUsdPointInstancerCollapsing::DontCollapse)
+		if (Settings && Settings->PointInstancerCollapsing == EUsdPointInstancerCollapsing::DontCollapse)
 		{
 			return false;
 		}
diff --git a/Engine/Plugins/Interchange/Extensions/OpenUSD/Source/Import/Public/InterchangeUsdTranslator.h b/Engine/Plugins/Interchange/Extensions/OpenUSD/Source/Import/Public/InterchangeUsdTranslator.h
--- a/Engine/Plugins/Interchange/Extensions/OpenUSD/Source/Import/Public/InterchangeUsdTranslator.h
+++ b/Engine/Plugins/Interchange/Extensions/OpenUSD/Source/Import/Public/InterchangeUsdTranslator.h
@@ -89,7 +89,7 @@ public:
 
 	/** How the USD Translator should behave when it tries collapsing a subtree and encounters a PointInstancer prim */
 	UPROPERTY(EditAnywhere, Category = "USD Translator|Collapsing")
-	EUsdPointInstancerCollapsing bPointInstancerCollapsing;
+	EUsdPointInstancerCollapsing PointInstancerCollapsing;
 
 	/**
 	 * Whether to use the dedicated collapsing schema to collapse prim subtrees into single meshes
