Interchange: Fix missing snippets when registering pipelines on the Interchange extension modules

#rnx
#Synced-CL 48758811

diff --git a/Engine/Plugins/Interchange/Extensions/OpenUSD/Source/Import/Private/InterchangeOpenUSDImportModule.cpp b/Engine/Plugins/Interchange/Extensions/OpenUSD/Source/Import/Private/InterchangeOpenUSDImportModule.cpp
--- a/Engine/Plugins/Interchange/Extensions/OpenUSD/Source/Import/Private/InterchangeOpenUSDImportModule.cpp
+++ b/Engine/Plugins/Interchange/Extensions/OpenUSD/Source/Import/Private/InterchangeOpenUSDImportModule.cpp
@@ -24,6 +24,7 @@
 #include "InterchangeManager.h"
 #include "InterchangeProjectSettings.h"
 
+#include "Engine/Engine.h"
 #include "Misc/CoreDelegates.h"
 #include "Modules/ModuleManager.h"
 
@@ -190,7 +191,14 @@ void FInterchangeOpenUSDImportModule::StartupModule()
 		}
 	};
 
-	PostEngineInitHandle = FCoreDelegates::OnPostEngineInit.AddLambda(RegisterItems);
+	if (GEngine)
+	{
+		RegisterItems();
+	}
+	else
+	{
+		PostEngineInitHandle = FCoreDelegates::OnPostEngineInit.AddLambda(RegisterItems);
+	}
 #endif	  // USE_USD_SDK
 }
 void FInterchangeOpenUSDImportModule::ShutdownModule()
diff --git a/Engine/Plugins/Interchange/Extensions/OpenVDB/Source/Import/Private/InterchangeOpenVDBImportModule.cpp b/Engine/Plugins/Interchange/Extensions/OpenVDB/Source/Import/Private/InterchangeOpenVDBImportModule.cpp
--- a/Engine/Plugins/Interchange/Extensions/OpenVDB/Source/Import/Private/InterchangeOpenVDBImportModule.cpp
+++ b/Engine/Plugins/Interchange/Extensions/OpenVDB/Source/Import/Private/InterchangeOpenVDBImportModule.cpp
@@ -108,7 +108,14 @@ void FInterchangeOpenVDBImportModule::StartupModule()
 #endif
 	};
 
-	FCoreDelegates::OnPostEngineInit.AddLambda(RegisterItems);
+	if (GEngine)
+	{
+		RegisterItems();
+	}
+	else
+	{
+		PostEngineInitHandle = FCoreDelegates::OnPostEngineInit.AddLambda(RegisterItems);
+	}
 }
 
 void FInterchangeOpenVDBImportModule::ShutdownModule()
