USD: Fix crash due to not reallocating the prim locks array when receiving a new value for NumPrimLocks

#Synced-CL 49553638

diff --git a/Engine/Plugins/Interchange/Extensions/OpenUSD/Source/Import/Private/InterchangeUSDInfoCache.cpp b/Engine/Plugins/Interchange/Extensions/OpenUSD/Source/Import/Private/InterchangeUSDInfoCache.cpp
--- a/Engine/Plugins/Interchange/Extensions/OpenUSD/Source/Import/Private/InterchangeUSDInfoCache.cpp
+++ b/Engine/Plugins/Interchange/Extensions/OpenUSD/Source/Import/Private/InterchangeUSDInfoCache.cpp
@@ -19,13 +19,15 @@
 #include "Misc/Paths.h"
 
 #define LOCTEXT_NAMESPACE "UsdInfoCache"
+#define DEFAULT_GEOMETRY_CACHE_MAX_DEPTH 15
+#define DEFAULT_NUM_PER_PRIM_LOCKS 32
 
 namespace UE::UsdInfoCache::Private
 {
 	// We fetch these from the cvar once before every Build()
 	static bool UseGeometryCacheUSD = true;
-	static int32 GeometryCacheMaxDepth = 15;
-	static int32 NumPerPrimLocks = 15;
+	static int32 GeometryCacheMaxDepth = DEFAULT_GEOMETRY_CACHE_MAX_DEPTH;
+	static int32 NumPerPrimLocks = DEFAULT_NUM_PER_PRIM_LOCKS;
 
 	// Flags to hint at the state of a prim for the purpose of geometry cache
 	enum class EGeometryCachePrimState : uint8
@@ -67,7 +69,6 @@ namespace UE::UsdInfoCache::Private
 			: AllowedExtensionsForGeometryCacheSource(UnrealUSDWrapper::GetNativeFileFormats())
 		{
 			AllowedExtensionsForGeometryCacheSource.Add(TEXT("abc"));
-			PrimLocks = new FRWLock[UE::UsdInfoCache::Private::NumPerPrimLocks]();
 		}
 
 		~FInterchangeUsdInfoCacheImpl()
@@ -101,6 +102,12 @@ namespace UE::UsdInfoCache::Private
 			return nullptr;
 		}
 
+		void ReallocatePrimLocks()
+		{
+			delete[] PrimLocks;
+			PrimLocks = new FRWLock[UE::UsdInfoCache::Private::NumPerPrimLocks]();
+		};
+
 		[[nodiscard]] FReadScopeLock LockForReading(const UE::UsdInfoCache::Private::FUsdPrimInfo& Info) const
 		{
 			return FReadScopeLock{PrimLocks[Info.PrimLockIndex]};
@@ -1084,10 +1091,10 @@ void FInterchangeUsdInfoCache::Build(UInterchangeUsdContext& Context)
 		UseGeometryCacheUSD = GeometryCacheEnableCVar && GeometryCacheEnableCVar->GetBool();
 
 		static IConsoleVariable* GeometryCacheDepthCVar = IConsoleManager::Get().FindConsoleVariable(TEXT("USD.GeometryCache.MaxDepth"));
-		GeometryCacheMaxDepth = GeometryCacheDepthCVar ? GeometryCacheDepthCVar->GetInt() : 15;
+		GeometryCacheMaxDepth = GeometryCacheDepthCVar ? GeometryCacheDepthCVar->GetInt() : DEFAULT_GEOMETRY_CACHE_MAX_DEPTH;
 
 		static IConsoleVariable* NumLocksCvar = IConsoleManager::Get().FindConsoleVariable(TEXT("USD.NumPerPrimLocks"));
-		NumPerPrimLocks = FMath::Max(NumLocksCvar ? NumLocksCvar->GetInt() : 32, 1);
+		NumPerPrimLocks = FMath::Max(NumLocksCvar ? NumLocksCvar->GetInt() : DEFAULT_NUM_PER_PRIM_LOCKS, 1);
 	}
 
 	Clear();
@@ -1116,6 +1123,9 @@ void FInterchangeUsdInfoCache::Build(UInterchangeUsdContext& Context)
 	{
 		TGuardValue<bool> Guard{ImplPtr->bIsBuilding, true};
 
+		// Do this after we potentially got a new value for NumPrimLocks from the cvar
+		ImplPtr->ReallocatePrimLocks();
+
 		PreallocateForSubtree(UE::FSdfPath::AbsoluteRootPath(), Stage, *ImplPtr);
 
 		CheckForGeometryCache(Context, *ImplPtr);
@@ -1164,6 +1174,8 @@ bool FInterchangeUsdInfoCache::IsEmpty()
 	return true;
 }
 
+#undef DEFAULT_GEOMETRY_CACHE_MAX_DEPTH
+#undef DEFAULT_NUM_PER_PRIM_LOCKS
 #undef LOCTEXT_NAMESPACE
 
 #endif	  // USE_USD_SDK
