USD Interchange: Allow reordering schema handlers in order to control render context priority

#Synced-CL 49553638

diff --git a/Engine/Plugins/Interchange/Extensions/OpenUSD/Source/Editor/Private/InterchangeUsdHandlerOrderWindow.cpp b/Engine/Plugins/Interchange/Extensions/OpenUSD/Source/Editor/Private/InterchangeUsdHandlerOrderWindow.cpp
new file mode 100644
--- /dev/null
+++ b/Engine/Plugins/Interchange/Extensions/OpenUSD/Source/Editor/Private/InterchangeUsdHandlerOrderWindow.cpp
@@ -0,0 +1,597 @@
+// Copyright Epic Games, Inc. All Rights Reserved.
+
+#include "InterchangeUsdHandlerOrderWindow.h"
+
+#include "InterchangeUsdTranslator.h"
+#include "SchemaHandlers/SchemaHandlerEntry.h"
+#include "SchemaHandlers/SchemaHandlerRegistry.h"
+#include "SchemaHandlers/SchemaHandlerUtils.h"
+#include "UnrealUSDWrapper.h"
+
+#include "CoreGlobals.h"
+#include "Framework/Application/SlateApplication.h"
+#include "Input/DragAndDrop.h"
+#include "Interfaces/IMainFrameModule.h"
+#include "Misc/ConfigCacheIni.h"
+#include "Misc/ScopedSlowTask.h"
+#include "Modules/ModuleManager.h"
+#include "PropertyEditorModule.h"
+#include "Widgets/Input/SButton.h"
+#include "Widgets/Input/SCheckBox.h"
+#include "Widgets/SBoxPanel.h"
+#include "Widgets/SWindow.h"
+#include "Widgets/Views/STableRow.h"
+
+#define LOCTEXT_NAMESPACE "InterchangeUsdHandlerOrderWindow"
+
+namespace UE::UsdHandlerOrderWindow::Private
+{
+	using namespace UE::Interchange::USD;
+
+	class SSchemaHandlerRow;
+
+	class FSchemaHandlerEntryDragDropOp : public FDragDropOperation
+	{
+	public:
+		DRAG_DROP_OPERATOR_TYPE(FSchemaHandlerEntryDragDropOp, FDragDropOperation)
+		DECLARE_DELEGATE(FOnDrop);
+
+		FOnDrop OnDropEvent;
+		TSharedPtr<SSchemaHandlerRow> DraggedRow;
+
+		FSchemaHandlerEntryDragDropOp()
+		{
+			MouseCursor = EMouseCursor::GrabHandClosed;
+		}
+
+		virtual TSharedPtr<SWidget> GetDefaultDecorator() const override
+		{
+			return SNullWidget::NullWidget;
+		}
+
+		virtual void OnDrop(bool bDropWasHandled, const FPointerEvent& MouseEvent) override
+		{
+			FDragDropOperation::OnDrop(bDropWasHandled, MouseEvent);
+			OnDropEvent.ExecuteIfBound();
+		}
+	};
+
+	// Swaps the empty string for 'universal' if it exists in the list
+	TArray<FString> GetUserFacingRenderContexts(TArray<FString> RenderContexts)
+	{
+		for (FString& RenderContext : RenderContexts)
+		{
+			if (RenderContext.IsEmpty())
+			{
+				RenderContext = UnrealIdentifiers::UniversalRenderContextDisplayString;
+			}
+		}
+
+		return RenderContexts;
+	}
+
+	class SSchemaHandlerRow : public STableRow<TSharedPtr<FSchemaHandlerEntry>>
+	{
+	public:
+		typedef typename STableRow<TSharedPtr<FSchemaHandlerEntry>>::FOnCanAcceptDrop FUsdOnCanAcceptDrop;
+		typedef typename STableRow<TSharedPtr<FSchemaHandlerEntry>>::FOnAcceptDrop FUsdOnAcceptDrop;
+
+		SLATE_BEGIN_ARGS(SSchemaHandlerRow)
+		{
+		}
+		SLATE_EVENT(FUsdOnCanAcceptDrop, OnCanAcceptDrop)
+		SLATE_EVENT(FUsdOnAcceptDrop, OnAcceptDrop)
+		SLATE_END_ARGS()
+
+		void Construct(
+			const FArguments& InArgs,
+			TSharedPtr<FSchemaHandlerEntry> InTreeItem,
+			const TSharedRef<STableViewBase>& OwnerTable
+		)
+		{
+			STableRow::Construct(
+				STableRow::FArguments()
+					.OnCanAcceptDrop(InArgs._OnCanAcceptDrop)
+					.OnAcceptDrop(InArgs._OnAcceptDrop)
+					.OnDragDetected(this, &SSchemaHandlerRow::OnDragDetected)
+					.Padding(FMargin(4.0f, 4.0f, 4.0f, 0.0f)), // Spacing between rows, as ListView doesn't provide this
+				OwnerTable
+			);
+
+			TreeItem = InTreeItem;
+
+			SetRowContent(
+				SNew(SBorder)
+				.BorderImage_Lambda([this]()
+				{
+					return bBeingDragged
+						? FAppStyle::GetBrush("Brushes.Recessed")
+						: IsHovered()
+							? FAppStyle::GetBrush("Brushes.Secondary")
+							: FAppStyle::GetBrush("ToolPanel.GroupBorder");
+				})
+				.ColorAndOpacity_Lambda([this]()
+				{
+					return (bBeingDragged || !TreeItem->bEnabled)
+						? FLinearColor{0.5f, 0.5f, 0.5f, 0.5f}
+						: FLinearColor{1.0f, 1.0f, 1.0f, 1.0f};
+				})
+				.Cursor_Lambda([this]()
+				{
+					return IsHovered() ? EMouseCursor::GrabHand : EMouseCursor::Default;
+				})
+				.HAlign(HAlign_Fill)
+				.VAlign(VAlign_Fill)
+				.Padding(FMargin(10.0f))
+				[
+					SNew(SHorizontalBox)
+					+SHorizontalBox::Slot()
+					.HAlign(HAlign_Left)
+					.VAlign(VAlign_Center)
+					.Padding(FMargin(0, 0, 10.0f, 0))
+					.MaxWidth(20)
+					.AutoWidth()
+					[
+						SNew(SCheckBox)
+						.ToolTipText(LOCTEXT("CheckBoxToolTip", "Enable or disable the schema handler"))
+						.Cursor(EMouseCursor::Default)
+						.IsChecked_Lambda([InTreeItem]()
+						{
+							if (InTreeItem)
+							{
+								return InTreeItem->bEnabled
+									? ECheckBoxState::Checked
+									: ECheckBoxState::Unchecked;
+							}
+	
+							return ECheckBoxState::Undetermined;
+						})
+						.OnCheckStateChanged_Lambda([InTreeItem](ECheckBoxState State)
+						{
+							if (InTreeItem)
+							{
+								if (State == ECheckBoxState::Checked)
+								{
+									InTreeItem->bEnabled = true;
+								}
+								else if (State == ECheckBoxState::Unchecked)
+								{
+									InTreeItem->bEnabled = false;
+								}
+							}
+						})
+					]
+	
+					+SHorizontalBox::Slot()
+					.HAlign(HAlign_Left)
+					.VAlign(VAlign_Center)
+					.FillWidth(1.0f)
+					[
+						SNew(STextBlock)
+						.Text(FText::FromString(TreeItem->HandlerName))
+						.ToolTipText(LOCTEXT("HandlerName", "Handler name"))
+						.Font(FAppStyle::GetFontStyle(TEXT("PropertyWindow.NormalFont")))
+					]
+
+					+SHorizontalBox::Slot()
+					.HAlign(HAlign_Left)
+					.VAlign(VAlign_Center)
+					.FillWidth(0.7)
+					[
+						SNew(STextBlock)
+						.Text(FText::FromString(TreeItem->SchemaName))
+						.ToolTipText(LOCTEXT("HandlerSchema", "Prim type name the handler applies to (\"IsA schema\").\nIn practice the handler may also filter for additional criteria like API schemas."))
+						.Font(FAppStyle::GetFontStyle(TEXT("PropertyWindow.NormalFont")))
+					]
+
+					+SHorizontalBox::Slot()
+					.HAlign(HAlign_Left)
+					.VAlign(VAlign_Center)
+					.FillWidth(0.7)
+					[
+						SNew(STextBlock)
+						.Text(FText::FromString(FString::Join(GetUserFacingRenderContexts(TreeItem->DefaultRenderContexts), TEXT(", "))))
+						.ToolTipText(LOCTEXT("RenderContext", "List of render contexts the handler will try parsing, in order of priority (first item of the list will be checked first)"))
+						.Visibility(TreeItem && TreeItem->SchemaName == TEXT("Material") ? EVisibility::Visible : EVisibility::Hidden)
+						.MinDesiredWidth(50.0f)
+						.Font(FAppStyle::GetFontStyle("PropertyWindow.NormalFont"))
+					]
+
+					+SHorizontalBox::Slot()
+					.HAlign(HAlign_Right)
+					.VAlign(VAlign_Center)
+					.MaxWidth(30.0f)
+					.AutoWidth()
+					[
+						SNew(SImage)
+						.ColorAndOpacity(FSlateColor::UseSubduedForeground())
+						.Image(FAppStyle::Get().GetBrush("Icons.DragHandle"))
+					]
+				]
+			);
+		}
+
+		FReply OnDragDetected(const FGeometry& Geometry, const FPointerEvent& PointerEvent)
+		{
+			TSharedRef<FSchemaHandlerEntryDragDropOp> Op = MakeShared<FSchemaHandlerEntryDragDropOp>();
+			Op->DraggedRow = SharedThis(this);
+
+			bBeingDragged = true;
+			Op->OnDropEvent.BindSPLambda(this, [this]()
+			{
+				bBeingDragged = false;
+			});
+
+			return FReply::Handled().BeginDragDrop(Op);
+		}
+
+		// We have to override this because the default table row implementation will only initiate drag and drop
+		// if the list selection mode is not None, but we don't really want to allow selection on our list view
+		virtual FReply OnMouseButtonDown(const FGeometry& MyGeometry, const FPointerEvent& MouseEvent) override
+		{
+			TSharedRef<ITypedTableView<TSharedPtr<FSchemaHandlerEntry>>> OwnerTable = OwnerTablePtr.Pin().ToSharedRef();
+
+			if (MouseEvent.GetEffectingButton() == EKeys::LeftMouseButton)
+			{
+				return FReply::Handled()
+					.DetectDrag(SharedThis(this), EKeys::LeftMouseButton)
+					.SetUserFocus(OwnerTable->AsWidget(), EFocusCause::Mouse)
+					.CaptureMouse(SharedThis(this));
+			}
+
+			return FReply::Unhandled();
+		}
+
+	public:
+		bool bBeingDragged = false;
+		TSharedPtr<FSchemaHandlerEntry> TreeItem;
+	};
+
+	TArray<TSharedPtr<FSchemaHandlerEntry>> WrapElementsWithSharedPtr(const TArray<FSchemaHandlerEntry>& Entries)
+	{
+		TArray<TSharedPtr<FSchemaHandlerEntry>> ResultEntries;
+		ResultEntries.Reserve(Entries.Num());
+		for (const FSchemaHandlerEntry& Entry : Entries)
+		{
+			ResultEntries.Add(MakeShared<FSchemaHandlerEntry>(Entry));
+		}
+		return ResultEntries;
+	}
+
+	TArray<FSchemaHandlerEntry> UnwrapElementsWithSharedPtr(const TArray<TSharedPtr<FSchemaHandlerEntry>>& Entries)
+	{
+		TArray<FSchemaHandlerEntry> ResultEntries;
+		ResultEntries.Reserve(Entries.Num());
+		for (const TSharedPtr<FSchemaHandlerEntry>& Entry : Entries)
+		{
+			if (Entry)
+			{
+				ResultEntries.Add(*Entry);
+			}
+		}
+		return ResultEntries;
+	}
+
+	void WriteHandlerEntriesToSettings(TArray<TSharedPtr<FSchemaHandlerEntry>> Entries, UInterchangeUsdTranslatorSettings* Settings)
+	{
+		if (!Settings)
+		{
+			return;
+		}
+
+		static FProperty* CustomEntriesProperty = UInterchangeUsdTranslatorSettings::StaticClass()->FindPropertyByName(
+			GET_MEMBER_NAME_CHECKED(UInterchangeUsdTranslatorSettings, CustomHandlerEntries)
+		);
+
+		FPropertyChangedEvent PropertyChangeEvent{CustomEntriesProperty};
+		Settings->Modify();
+		{
+			Settings->CustomHandlerEntries = UnwrapElementsWithSharedPtr(Entries);
+		}
+		Settings->PostEditChangeProperty(PropertyChangeEvent);
+	}
+
+	TArray<TSharedPtr<FSchemaHandlerEntry>> ReadHandlerEntriesFromSettings(UInterchangeUsdTranslatorSettings* Settings)
+	{
+		if (!Settings)
+		{
+			return {};
+		}
+
+		// Try applying the existing reordering saved on the translator settings, if any
+		if (Settings->CustomHandlerEntries.Num() > 0)
+		{
+			return WrapElementsWithSharedPtr(Settings->CustomHandlerEntries);
+		}
+
+#if USE_USD_SDK
+		return WrapElementsWithSharedPtr(FSchemaHandlerRegistry::RegisteredHandlerEntries);
+#else
+		return {};
+#endif	  // USE_USD_SDK
+	}
+}	 // namespace UE::UsdHandlerOrderWindow::Private
+
+bool SInterchangeUsdHandlerOrderWindow::ShowWindow(UInterchangeUsdTranslatorSettings* InSettings)
+{
+	// Reference: SUsdOptionsWindow::ShowOptions
+
+	using namespace UE::UsdHandlerOrderWindow::Private;
+
+	if (!InSettings)
+	{
+		return false;
+	}
+
+	const static FString ConfigSectionName = TEXT("InterchangeUsdHandlerOrderWindow");
+	const static FString WidthConfigName = TEXT("DialogWidth");
+	const static FString HeightConfigName = TEXT("DialogHeight");
+
+	// Grab our desired window size from persistent storage, if we have any
+	int32 Width = 540;
+	int32 Height = 860;
+	{
+		FDisplayMetrics DisplayMetrics;
+		FSlateApplication::Get().GetDisplayMetrics(DisplayMetrics);
+
+		if (GConfig->GetInt(*ConfigSectionName, *WidthConfigName, Width, GEditorPerProjectIni))
+		{
+			// Just to make sure we don't load something unusable
+			Width = FMath::Min(FMath::Max(Width, 100), DisplayMetrics.PrimaryDisplayWidth);
+		}
+
+		if (GConfig->GetInt(*ConfigSectionName, *HeightConfigName, Height, GEditorPerProjectIni))
+		{
+			Height = FMath::Min(FMath::Max(Height, 100), DisplayMetrics.PrimaryDisplayHeight);
+		}
+	}
+
+	TSharedRef<SInterchangeUsdHandlerOrderWindow> OrderWidget = SNew(SInterchangeUsdHandlerOrderWindow);
+
+	TSharedRef<SWindow> Window = SNew(SWindow)
+		.Title(LOCTEXT("WindowTitle", "USD Schema Handler Order"))
+		.SizingRule(ESizingRule::UserSized)
+		.AdjustInitialSizeAndPositionForDPIScale(false)
+		.MinWidth(540.0f)  // Any smaller and it clips the title
+		.ClientSize(FVector2D{static_cast<double>(Width), static_cast<double>(Height)});
+	Window->SetContent(OrderWidget);
+	Window->GetOnWindowClosedEvent().AddLambda(
+		[OrderWidget](const TSharedRef<SWindow>& Window)
+		{
+			// Store our last window size on persistent storage
+
+			// We use the geometry here because any sort of size we can extract from the window includes some added margins
+			// that are not easy to account for, even with GetWindowSizeFromClientSize. The original provided ClientSize
+			// seems to always match the content geometry though, so we're saving and loading the same thing
+			FGeometry Geometry = Window->GetContent()->GetCachedGeometry();
+			FVector2D ContentSize = Geometry.Size * Geometry.Scale;
+			const int32 Width = static_cast<int32>(ContentSize.X);
+			const int32 Height = static_cast<int32>(ContentSize.Y);
+
+			GConfig->SetInt(*ConfigSectionName, *WidthConfigName, Width, GEditorPerProjectIni);
+			GConfig->SetInt(*ConfigSectionName, *HeightConfigName, Height, GEditorPerProjectIni);
+
+			const bool bRemoveFromCache = false;
+			GConfig->Flush(bRemoveFromCache, GEditorPerProjectIni);
+		}
+	);
+
+	OrderWidget->Settings = InSettings;
+	OrderWidget->WeakWindow = Window;
+	OrderWidget->Entries = ReadHandlerEntriesFromSettings(InSettings);
+
+	// Preemptively make sure we have a progress dialog created before showing our modal. This because the progress
+	// dialog itself is also modal. If it doesn't exist yet, and our options dialog causes a progress dialog
+	// to be spawned (e.g. when switching the Level to export via the LevelSequenceUSDExporter), the progress dialog
+	// will be pushed to the end of FSlateApplication::ActiveModalWindows (SlateApplication.cpp) and cause our options
+	// dialog to pop out of its modal loop (FSlateApplication::AddModalWindow), instantly returning false to our caller
+	FScopedSlowTask Progress(1, LOCTEXT("ShowingDialog", "Picking order..."));
+	Progress.MakeDialogDelayed(0.25f);
+
+	TSharedPtr<SWindow> ParentWindow;
+	if (FModuleManager::Get().IsModuleLoaded("MainFrame"))
+	{
+		IMainFrameModule& MainFrame = FModuleManager::LoadModuleChecked<IMainFrameModule>("MainFrame");
+		ParentWindow = MainFrame.GetParentWindow();
+	}
+
+	// Note: Being a modal window this will only return when the window closes
+	const bool bSlowTaskWindow = false;
+	FSlateApplication::Get().AddModalWindow(Window, ParentWindow, bSlowTaskWindow);
+
+	return OrderWidget->bNeedNewTranslation;
+}
+
+void SInterchangeUsdHandlerOrderWindow::Construct(const FArguments& InArgs)
+{
+	using namespace UE::UsdHandlerOrderWindow::Private;
+
+	ChildSlot
+	[
+		SNew(SBorder)
+		.BorderImage(FAppStyle::GetBrush("ToolPanel.DarkGroupBorder"))
+		[
+			SNew(SVerticalBox)
+
+			+SVerticalBox::Slot()
+			.AutoHeight()
+			.VAlign(VAlign_Center)
+			[
+				SNew(STextBlock)
+				.Text(LOCTEXT("MainInfoText", "All enabled handlers below are executed from top to bottom when translating each prim.\nDrag and drop to reorder them.\n\nHover this text for more info."))
+				.ToolTipText(LOCTEXT("MainHoverToolTip", "A handler is a C++ class that is registered with the USD translator, and invoked when translating prims of particular schemas.\n\nHandlers that run first generally get to pick the Interchange nodes that are produced, while handlers that run later get to override data on those nodes.\n\nFor example, reorder the material handlers to pick the priority order for parsing material render contexts.\nSome material handlers also allow specifying additional render contexts that they will try parsing the material with."))
+				.AutoWrapText(true)
+				.Margin(10.0f)
+			]
+
+			+SVerticalBox::Slot()
+			.FillHeight(1.0f)
+			[
+				SNew(SBox)
+				.Padding(FMargin(10.0f, 0.0f, 10.0f, 0.0f))
+				[
+					SAssignNew(ListView, SListView<TSharedPtr<FSchemaHandlerEntry>>)
+					.ListItemsSource(&Entries)
+					.OnGenerateRow(this, &SInterchangeUsdHandlerOrderWindow::OnGenerateListRow)
+					.SelectionMode(ESelectionMode::None)
+				]
+			]
+
+			+SVerticalBox::Slot()
+			.AutoHeight()
+			.HAlign(HAlign_Fill)
+			[
+				SNew(SHorizontalBox)
+				+ SHorizontalBox::Slot()
+				.AutoWidth()
+				.HAlign(HAlign_Left)
+				.VAlign(VAlign_Center)
+				[
+					SNew(SButton)
+					.Text(LOCTEXT("ResetToDefaultButton", "Reset To Default"))
+					.OnClicked(this, &SInterchangeUsdHandlerOrderWindow::OnReset)
+				]
+
+				+ SHorizontalBox::Slot()
+				.FillWidth(1.0f)
+				.VAlign(VAlign_Center)
+				.HAlign(HAlign_Right)
+				[
+					SNew(SHorizontalBox)
+					+ SHorizontalBox::Slot()
+					.Padding(FAppStyle::Get().GetMargin("StandardDialog.SlotPadding"))
+					.AutoWidth()
+					[
+						SNew(SButton)
+						.ButtonStyle(&FAppStyle::Get().GetWidgetStyle<FButtonStyle>("PrimaryButton"))
+						.Text(LOCTEXT("OKButton", "OK"))
+						.OnClicked(this, &SInterchangeUsdHandlerOrderWindow::OnAccept)
+					]
+					+ SHorizontalBox::Slot()
+					.Padding(FAppStyle::Get().GetMargin("StandardDialog.SlotPadding"))
+					.AutoWidth()
+					[
+						SNew(SButton)
+						.ButtonStyle(&FAppStyle::Get().GetWidgetStyle<FButtonStyle>("Button"))
+						.Text(LOCTEXT("CancelButton", "Cancel"))
+						.OnClicked(this, &SInterchangeUsdHandlerOrderWindow::OnCancel)
+					]
+				]
+			]
+		]
+	];
+}
+
+TSharedRef<ITableRow> SInterchangeUsdHandlerOrderWindow::OnGenerateListRow(TSharedPtr<FSchemaHandlerEntry> Item, const TSharedRef<STableViewBase>& OwnerTable)
+{
+	using namespace UE::UsdHandlerOrderWindow::Private;
+
+	return SNew(SSchemaHandlerRow, Item, OwnerTable)
+		.OnCanAcceptDrop(this, &SInterchangeUsdHandlerOrderWindow::OnRowCanAcceptDrop)
+		.OnAcceptDrop(this, &SInterchangeUsdHandlerOrderWindow::OnRowAcceptDrop);
+}
+
+FReply SInterchangeUsdHandlerOrderWindow::OnAccept()
+{
+	using namespace UE::UsdHandlerOrderWindow::Private;
+
+	WriteHandlerEntriesToSettings(Entries, Settings);
+
+	// For now let's always retranslate if we hit OK, for simplicity
+	bNeedNewTranslation = true;
+
+	if (WeakWindow.IsValid())
+	{
+		WeakWindow.Pin()->RequestDestroyWindow();
+	}
+
+	return FReply::Handled();
+}
+
+FReply SInterchangeUsdHandlerOrderWindow::OnCancel()
+{
+	bNeedNewTranslation = false;
+
+	if (WeakWindow.IsValid())
+	{
+		WeakWindow.Pin()->RequestDestroyWindow();
+	}
+
+	return FReply::Handled();
+}
+
+FReply SInterchangeUsdHandlerOrderWindow::OnReset()
+{
+	using namespace UE::UsdHandlerOrderWindow::Private;
+
+#if USE_USD_SDK
+	Entries = WrapElementsWithSharedPtr(FSchemaHandlerRegistry::RegisteredHandlerEntries);
+#else
+	Entries = {};
+#endif	  // USE_USD_SDK
+
+	ListView->RequestListRefresh();
+
+	bNeedNewTranslation = true;
+
+	return FReply::Handled();
+}
+
+FReply SInterchangeUsdHandlerOrderWindow::OnKeyDown(const FGeometry& MyGeometry, const FKeyEvent& InKeyEvent)
+{
+	if (InKeyEvent.GetKey() == EKeys::Escape)
+	{
+		return OnCancel();
+	}
+
+	return FReply::Unhandled();
+}
+
+TOptional<EItemDropZone> SInterchangeUsdHandlerOrderWindow::OnRowCanAcceptDrop(
+	const FDragDropEvent& Event,
+	EItemDropZone Zone,
+	TSharedPtr<FSchemaHandlerEntry> Item
+)
+{
+	using namespace UE::UsdHandlerOrderWindow::Private;
+
+	if (TSharedPtr<FSchemaHandlerEntryDragDropOp> Op = Event.GetOperationAs<FSchemaHandlerEntryDragDropOp>())
+	{
+		if (Op->DraggedRow && Op->DraggedRow->TreeItem != Item)
+		{
+			// Note that we're returning BelowItem even if we're hovering OntoItem here: That makes it a bit easier
+			// to use as you don't have to aim for exactly the small spaces between the items
+			return Zone == EItemDropZone::AboveItem ? EItemDropZone::AboveItem : EItemDropZone::BelowItem;
+		}
+	}
+
+	return {};
+}
+
+FReply SInterchangeUsdHandlerOrderWindow::OnRowAcceptDrop(const FDragDropEvent& Event, EItemDropZone Zone, TSharedPtr<FSchemaHandlerEntry> Item)
+{
+	using namespace UE::UsdHandlerOrderWindow::Private;
+
+	if (OnRowCanAcceptDrop(Event, Zone, Item))
+	{
+		if (TSharedPtr<FSchemaHandlerEntryDragDropOp> Op = Event.GetOperationAs<FSchemaHandlerEntryDragDropOp>())
+		{
+			const TSharedPtr<FSchemaHandlerEntry>& DraggedEntry = Op->DraggedRow->TreeItem;
+			Entries.Remove(DraggedEntry);
+
+			const int32 ItemIndex = Entries.IndexOfByKey(Item);
+			if (ensure(ItemIndex != INDEX_NONE))
+			{
+				// Here we watch out for ending up with -1 in case we're inserting above the first entry.
+				const int32 TargetIndex = FMath::Clamp(ItemIndex + (Zone == EItemDropZone::AboveItem ? 0 : 1), 0, Entries.Num());
+				Entries.Insert(DraggedEntry, TargetIndex);
+			}
+
+			ListView->RebuildList();
+		}
+	}
+
+	return FReply::Unhandled();
+}
+
+#undef LOCTEXT_NAMESPACE
diff --git a/Engine/Plugins/Interchange/Extensions/OpenUSD/Source/Editor/Private/InterchangeUsdHandlerOrderWindow.h b/Engine/Plugins/Interchange/Extensions/OpenUSD/Source/Editor/Private/InterchangeUsdHandlerOrderWindow.h
new file mode 100644
--- /dev/null
+++ b/Engine/Plugins/Interchange/Extensions/OpenUSD/Source/Editor/Private/InterchangeUsdHandlerOrderWindow.h
@@ -0,0 +1,45 @@
+// Copyright Epic Games, Inc. All Rights Reserved.
+
+#pragma once
+
+#include "Widgets/SCompoundWidget.h"
+#include "Widgets/SWindow.h"
+#include "Widgets/Views/SListView.h"
+#include "Widgets/Views/STableRow.h"
+
+class ITableRow;
+class STableViewBase;
+class UInterchangeUsdTranslatorSettings;
+struct FSchemaHandlerEntry;
+
+class SInterchangeUsdHandlerOrderWindow : public SCompoundWidget
+{
+public:
+	SLATE_BEGIN_ARGS(SInterchangeUsdHandlerOrderWindow)
+	{
+	}
+	SLATE_END_ARGS()
+
+public:
+	static bool ShowWindow(UInterchangeUsdTranslatorSettings* Settings);
+
+	void Construct(const FArguments& InArgs);
+	TSharedRef<ITableRow> OnGenerateListRow(TSharedPtr<FSchemaHandlerEntry> Item, const TSharedRef<STableViewBase>& OwnerTable);
+
+	FReply OnAccept();
+	FReply OnCancel();
+	FReply OnReset();
+	virtual FReply OnKeyDown(const FGeometry& MyGeometry, const FKeyEvent& InKeyEvent) override;
+
+	TOptional<EItemDropZone> OnRowCanAcceptDrop(const FDragDropEvent& Event, EItemDropZone Zone, TSharedPtr<FSchemaHandlerEntry> Item);
+	FReply OnRowAcceptDrop(const FDragDropEvent& Event, EItemDropZone Zone, TSharedPtr<FSchemaHandlerEntry> Item);
+	FReply OnRowDragDetected(const FGeometry& Geometry, const FPointerEvent& PointerEvent);
+	void OnRowDragLeave(const FDragDropEvent& Event);
+
+private:
+	bool bNeedNewTranslation = false;
+	UInterchangeUsdTranslatorSettings* Settings;
+	TArray<TSharedPtr<FSchemaHandlerEntry>> Entries;
+	TWeakPtr<SWindow> WeakWindow;
+	TSharedPtr<SListView<TSharedPtr<FSchemaHandlerEntry>>> ListView;
+};
diff --git a/Engine/Plugins/Interchange/Extensions/OpenUSD/Source/Editor/Private/InterchangeUsdTranslatorBlueprintLibrary.cpp b/Engine/Plugins/Interchange/Extensions/OpenUSD/Source/Editor/Private/InterchangeUsdTranslatorBlueprintLibrary.cpp
new file mode 100644
--- /dev/null
+++ b/Engine/Plugins/Interchange/Extensions/OpenUSD/Source/Editor/Private/InterchangeUsdTranslatorBlueprintLibrary.cpp
@@ -0,0 +1,14 @@
+// Copyright Epic Games, Inc. All Rights Reserved.
+
+#include "InterchangeUsdTranslatorBlueprintLibrary.h"
+
+#include "SchemaHandlers/SchemaHandlerRegistry.h"
+
+TArray<FSchemaHandlerEntry> UInterchangeUsdTranslatorBlueprintLibrary::GetDefaultSchemaHandlerEntries()
+{
+#if USE_USD_SDK
+	return UE::Interchange::USD::FSchemaHandlerRegistry::RegisteredHandlerEntries;
+#else
+	return {};
+#endif	  // USE_USD_SDK
+}
diff --git a/Engine/Plugins/Interchange/Extensions/OpenUSD/Source/Editor/Private/InterchangeUsdTranslatorSettingsCustomization.cpp b/Engine/Plugins/Interchange/Extensions/OpenUSD/Source/Editor/Private/InterchangeUsdTranslatorSettingsCustomization.cpp
--- a/Engine/Plugins/Interchange/Extensions/OpenUSD/Source/Editor/Private/InterchangeUsdTranslatorSettingsCustomization.cpp
+++ b/Engine/Plugins/Interchange/Extensions/OpenUSD/Source/Editor/Private/InterchangeUsdTranslatorSettingsCustomization.cpp
@@ -2,6 +2,7 @@
 
 #include "InterchangeUsdTranslatorSettingsCustomization.h"
 
+#include "InterchangeUsdHandlerOrderWindow.h"
 #include "InterchangeUsdTranslator.h"
 #include "UnrealUSDWrapper.h"
 #include "USDMaterialUtils.h"
@@ -15,6 +16,7 @@
 #include "Styling/AppStyle.h"
 #include "UObject/Object.h"
 #include "Widgets/DeclarativeSyntaxSupport.h"
+#include "Widgets/Input/SButton.h"
 #include "Widgets/Input/SEditableTextBox.h"
 
 #define LOCTEXT_NAMESPACE "InterchangeUsdTranslatorSettingsCustomization"
@@ -50,79 +52,46 @@ void FInterchangeUsdTranslatorSettingsCustomization::CustomizeDetails(IDetailLay
 
 	IDetailCategoryBuilder& CatBuilder = DetailLayoutBuilder.EditCategory(TEXT("USD Translator"));
 
-	if (TSharedPtr<IPropertyHandle> RenderContextProperty = DetailLayoutBuilder.GetProperty(
-			GET_MEMBER_NAME_CHECKED(UInterchangeUsdTranslatorSettings, RenderContext)
+	if (TSharedPtr<IPropertyHandle> CustomHandlerProperty = DetailLayoutBuilder.GetProperty(
+			GET_MEMBER_NAME_CHECKED(UInterchangeUsdTranslatorSettings, CustomHandlerEntries)
 		))
 	{
-		DetailLayoutBuilder.HideProperty(RenderContextProperty);
+		DetailLayoutBuilder.HideProperty(CustomHandlerProperty);
 
-		// clang-format off
-		CatBuilder.AddCustomRow(FText::FromString(TEXT("RenderContextCustomization")))
-		.NameContent()
-		[
-			SNew(STextBlock)
-			.Text(FText::FromString(TEXT("Render Context")))
-			.Font(FAppStyle::GetFontStyle(TEXT("PropertyWindow.NormalFont")))
-			.ToolTipText(RenderContextProperty->GetToolTipText())
-		]
-		.ValueContent()
+		UInterchangeUsdTranslatorSettings* Settings = CurrentOptions;
+
+		CatBuilder.AddCustomRow(FText::FromString(TEXT("ConfigureHandlersButton")))
+		.WholeRowContent()
 		[
 			SNew(SBox)
-				.VAlign(VAlign_Center)
-				[
-					SNew(SComboBox<TSharedPtr<FString>>)
-						.OptionsSource(&RenderContextComboBoxItems)
-						.OnComboBoxOpening_Lambda([this]()
-						{
-							// We used to pull these from UsdUnreal::MaterialUtils::GetRegisteredRenderContexts(), but these are
-							// hard-coded for Interchange because that list may contain render contexts added in through the
-							// legacy USD render context registry, which we don't necessarily support on Interchange just yet (e.g. "mdl")
-							RenderContextComboBoxItems = {
-								MakeShared<FString>(TEXT("mtlx")),
-								MakeShared<FString>(TEXT("universal")),
-								MakeShared<FString>(TEXT("unreal")),
-							};
-						})
-						.OnGenerateWidget_Lambda([&](TSharedPtr<FString> Option)
-						{
-							TSharedPtr<SWidget> Widget = SNullWidget::NullWidget;
-							if (Option)
-							{
-								Widget = SNew(STextBlock)
-									.Text(FText::FromString(*Option))
-									.Font(FAppStyle::GetFontStyle("PropertyWindow.NormalFont"));
-							}
+			.VAlign(VAlign_Center)
+			.HAlign(HAlign_Center)
+			.Padding(5.0f)
+			[
+				SNew(SButton)
+				.HAlign(HAlign_Center)
+				.Text(LOCTEXT("EditSchemaHandlerButton", "Configure schema handlers"))
+				.OnClicked_Lambda([Settings, CustomHandlerProperty]()
+				{
+					FScopedTransaction Transaction(LOCTEXT("CustomHandlerTransaction", "Edit Custom Handlers"));
+					CustomHandlerProperty->NotifyPreChange();
 
-							return Widget.ToSharedRef();
-						})
-						.OnSelectionChanged_Lambda([this, RenderContextProperty](TSharedPtr<FString> ChosenOption, ESelectInfo::Type SelectInfo)
-						{
-							if (CurrentOptions && ChosenOption)
-							{
-								FScopedTransaction Transaction(LOCTEXT("RenderContextTransaction", "Edit Render Context"));
-								RenderContextProperty->NotifyPreChange();
-								{
-									CurrentOptions->RenderContext = **ChosenOption;
-								}
-								RenderContextProperty->NotifyPostChange(EPropertyChangeType::ValueSet);
-								RenderContextProperty->NotifyFinishedChangingProperties();
-							}
+					bool bNeedNewTranslation = SInterchangeUsdHandlerOrderWindow::ShowWindow(Settings);
 
-						})
-						[
-							SNew(STextBlock)
-							.Text_Lambda([this]() -> FText
-							{
-								if (CurrentOptions)
-								{
-									return FText::FromString(CurrentOptions->RenderContext.ToString());
-								}
-
-								return FText::GetEmpty();
-							})
-							.Font(FAppStyle::GetFontStyle(TEXT("PropertyWindow.NormalFont")))
-						]
-				]
+					if (bNeedNewTranslation)
+					{
+						// These events are important as they're the thing that eventually causes
+						// SInterchangeTranslatorSettingsDialog to trigger OnTranslatorSettingsDialogClosed, and
+						// ultimately cause the code in SInterchangePipelineConfigurationDialog::OnEditTranslatorSettings()
+						// to trigger a new translation of the source file with the new translator (and schema handler)
+						// settings when bTranslatorSettingsChanged is true
+						CustomHandlerProperty->NotifyPostChange(EPropertyChangeType::ValueSet);
+						CustomHandlerProperty->NotifyFinishedChangingProperties();
+					}
+
+					return FReply::Handled();
+				})
+			]
 		];
 	}
 
diff --git a/Engine/Plugins/Interchange/Extensions/OpenUSD/Source/Editor/Public/InterchangeUsdTranslatorBlueprintLibrary.h b/Engine/Plugins/Interchange/Extensions/OpenUSD/Source/Editor/Public/InterchangeUsdTranslatorBlueprintLibrary.h
new file mode 100644
--- /dev/null
+++ b/Engine/Plugins/Interchange/Extensions/OpenUSD/Source/Editor/Public/InterchangeUsdTranslatorBlueprintLibrary.h
@@ -0,0 +1,31 @@
+// Copyright Epic Games, Inc. All Rights Reserved.
+
+#pragma once
+
+#include "SchemaHandlers/SchemaHandlerEntry.h"
+
+#include "Engine/EngineTypes.h"
+#include "Kismet/BlueprintFunctionLibrary.h"
+
+#include "InterchangeUsdTranslatorBlueprintLibrary.generated.h"
+
+#define UE_API INTERCHANGEOPENUSDEDITOR_API
+
+UCLASS(MinimalAPI, meta = (ScriptName = "InterchangeUsdTranslatorLibrary"))
+class UInterchangeUsdTranslatorBlueprintLibrary : public UBlueprintFunctionLibrary
+{
+	GENERATED_BODY()
+
+public:
+	/**
+	 * Returns the default list of schema handlers entries registered via C++, i.e. FSchemaHandlerRegistry::RegisteredHandlerEntries.
+	 *
+	 * The main use case here is to retrieve this list, modify it as needed, and to set it as the CustomHandlerEntries property on
+	 * the UInterchangeUsdTranslatorSettings class default object, which will be used by any Interchange USD translations after
+	 * that point
+	 */
+	UFUNCTION(BlueprintCallable, Category = "USD Translator")
+	static UE_API TArray<FSchemaHandlerEntry> GetDefaultSchemaHandlerEntries();
+};
+
+#undef UE_API
diff --git a/Engine/Plugins/Interchange/Extensions/OpenUSD/Source/Import/Private/InterchangeOpenUSDImportModule.cpp b/Engine/Plugins/Interchange/Extensions/OpenUSD/Source/Import/Private/InterchangeOpenUSDImportModule.cpp
--- a/Engine/Plugins/Interchange/Extensions/OpenUSD/Source/Import/Private/InterchangeOpenUSDImportModule.cpp
+++ b/Engine/Plugins/Interchange/Extensions/OpenUSD/Source/Import/Private/InterchangeOpenUSDImportModule.cpp
@@ -14,6 +14,7 @@
 #include "SchemaHandlers/NaniteAssemblySchemaHandler.h"
 #include "SchemaHandlers/PointInstancerSchemaHandler.h"
 #include "SchemaHandlers/SchemaHandler.h"
+#include "SchemaHandlers/SchemaHandlerRegistry.h"
 #include "SchemaHandlers/SkeletonSchemaHandler.h"
 #include "SchemaHandlers/SpatialAudioSchemaHandler.h"
 #include "SchemaHandlers/UniversalMaterialSchemaHandler.h"
@@ -207,9 +208,9 @@ void FInterchangeOpenUSDImportModule::ShutdownModule()
 	using namespace UE::Interchange;
 	using namespace UE::Interchange::USD;
 
-	for (FSchemaHandlerIdentifier Id : RegisteredHandlers)
+	for (const FString& Name : RegisteredHandlers)
 	{
-		FSchemaHandlerRegistry::Unregister(Id);
+		FSchemaHandlerRegistry::Unregister(Name);
 	}
 
 	FCoreDelegates::OnPostEngineInit.Remove(PostEngineInitHandle);
diff --git a/Engine/Plugins/Interchange/Extensions/OpenUSD/Source/Import/Private/InterchangeUsdContext.cpp b/Engine/Plugins/Interchange/Extensions/OpenUSD/Source/Import/Private/InterchangeUsdContext.cpp
--- a/Engine/Plugins/Interchange/Extensions/OpenUSD/Source/Import/Private/InterchangeUsdContext.cpp
+++ b/Engine/Plugins/Interchange/Extensions/OpenUSD/Source/Import/Private/InterchangeUsdContext.cpp
@@ -278,9 +278,6 @@ void UInterchangeUsdContext::SetupMeshConversionOptions()
 		// really do anything anyway as the module doesn't use IMPLEMENT_MODULE_USD!
 		// Luckily we can get around this here because pxr::TfToken doesn't allocate on its own: At most USD makes a copy of the string, which it
 		// should allocate/deallocate on its own allocator.
-		CachedMeshConversionOptions.RenderContext = Settings->RenderContext == UnrealIdentifiers::UniversalRenderContext
-														? pxr::UsdShadeTokens->universalRenderContext
-														: pxr::TfToken{TCHAR_TO_UTF8(*Settings->RenderContext.ToString())};
 		CachedMeshConversionOptions.MaterialPurpose = Settings->MaterialPurpose.IsNone()
 														  ? pxr::UsdShadeTokens->allPurpose
 														  : pxr::TfToken{TCHAR_TO_UTF8(*Settings->MaterialPurpose.ToString())};
diff --git a/Engine/Plugins/Interchange/Extensions/OpenUSD/Source/Import/Private/InterchangeUsdTranslator.cpp b/Engine/Plugins/Interchange/Extensions/OpenUSD/Source/Import/Private/InterchangeUsdTranslator.cpp
--- a/Engine/Plugins/Interchange/Extensions/OpenUSD/Source/Import/Private/InterchangeUsdTranslator.cpp
+++ b/Engine/Plugins/Interchange/Extensions/OpenUSD/Source/Import/Private/InterchangeUsdTranslator.cpp
@@ -4,6 +4,7 @@
 
 #include "SchemaHandlers/HandlerAccumulatedInfo.h"
 #include "SchemaHandlers/SchemaHandler.h"
+#include "SchemaHandlers/SchemaHandlerRegistry.h"
 #include "SchemaHandlers/SchemaHandlerUtils.h"
 #include "UnrealUSDWrapper.h"
 #include "Usd/InterchangeUsdDefinitions.h"
@@ -235,7 +236,7 @@ TOptional<UE::Interchange::USD::FHandlerAccumulatedInfo> UInterchangeUSDTranslat
 
 		for (const TSharedRef<FSchemaHandler>& Handler : ImplPtr->SchemaHandlers)
 		{
-			if (Handler->CanHandlePrim(Prim, *UsdContext))
+			if (Handler->IsEnabled() && Handler->CanHandlePrim(Prim, *UsdContext))
 			{
 				Handler->OnTranslate(Prim, TraversalInfo, AccumulatedInfo, *UsdContext);
 			}
@@ -406,14 +407,20 @@ bool UInterchangeUSDTranslator::Translate(UInterchangeBaseNodeContainer& NodeCon
 	Context->USDZFilePath = USDZFilePath;
 	Context->DecompressedUSDZRoot = DecompressedUSDZRoot;
 
-	// Generate handlers for the translation
 	Impl = MakeUnique<UE::InterchangeUsdTranslator::Private::UInterchangeUSDTranslatorImpl>();
 	UInterchangeUSDTranslatorImpl* ImplPtr = Impl.Get();
 	if (!ImplPtr)
 	{
 		return false;
 	}
-	ImplPtr->SchemaHandlers = FSchemaHandlerRegistry::GenerateHandlers();
+
+	// Generate handlers for the translation
+	const TArray<FSchemaHandlerEntry>* HandlerEntries = &FSchemaHandlerRegistry::RegisteredHandlerEntries;
+	if (Settings && Settings->CustomHandlerEntries.Num() > 0)
+	{
+		HandlerEntries = &Settings->CustomHandlerEntries;
+	}
+	ImplPtr->SchemaHandlers = FSchemaHandlerRegistry::GenerateHandlers(HandlerEntries);
 
 	// Build info cache now that everything is setup
 	Context->SetupInterchangeInfoCache();
@@ -520,7 +527,10 @@ TOptional<UE::Interchange::FMeshPayloadData> UInterchangeUSDTranslator::GetMeshP
 	{
 		for (const TSharedRef<FSchemaHandler>& Handler : ImplPtr->SchemaHandlers)
 		{
-			Handler->OnGetMeshPayloadData(PayloadKey, PayloadAttributes, *UsdContext, Result);
+			if (Handler->IsEnabled())
+			{
+				Handler->OnGetMeshPayloadData(PayloadKey, PayloadAttributes, *UsdContext, Result);
+			}
 		}
 	}
 #endif	  // USE_USD_SDK
@@ -546,7 +556,10 @@ TOptional<UE::Interchange::FImportImage> UInterchangeUSDTranslator::GetTexturePa
 	{
 		for (const TSharedRef<FSchemaHandler>& Handler : ImplPtr->SchemaHandlers)
 		{
-			Handler->OnGetTexturePayloadData(PayloadKey, AlternateTexturePath, *UsdContext, Result);
+			if (Handler->IsEnabled())
+			{
+				Handler->OnGetTexturePayloadData(PayloadKey, AlternateTexturePath, *UsdContext, Result);
+			}
 		}
 	}
 #endif	  // USE_USD_SDK
@@ -572,7 +585,10 @@ TOptional<UE::Interchange::FImportBlockedImage> UInterchangeUSDTranslator::GetBl
 	{
 		for (const TSharedRef<FSchemaHandler>& Handler : ImplPtr->SchemaHandlers)
 		{
-			Handler->OnGetBlockedTexturePayloadData(PayloadKey, AlternateTexturePath, *UsdContext, Result);
+			if (Handler->IsEnabled())
+			{
+				Handler->OnGetBlockedTexturePayloadData(PayloadKey, AlternateTexturePath, *UsdContext, Result);
+			}
 		}
 	}
 #endif	  // USE_USD_SDK
@@ -597,7 +613,10 @@ TArray<UE::Interchange::FAnimationPayloadData> UInterchangeUSDTranslator::GetAni
 	{
 		for (const TSharedRef<FSchemaHandler>& Handler : ImplPtr->SchemaHandlers)
 		{
-			Handler->OnGetAnimationPayloadData(PayloadQueries, *UsdContext, Result);
+			if (Handler->IsEnabled())
+			{
+				Handler->OnGetAnimationPayloadData(PayloadQueries, *UsdContext, Result);
+			}
 		}
 	}
 #endif	  // USE_USD_SDK
@@ -620,7 +639,10 @@ TOptional<UE::Interchange::FVolumePayloadData> UInterchangeUSDTranslator::GetVol
 	{
 		for (const TSharedRef<FSchemaHandler>& Handler : ImplPtr->SchemaHandlers)
 		{
-			Handler->OnGetVolumePayloadData(PayloadKey, *UsdContext, Result);
+			if (Handler->IsEnabled())
+			{
+				Handler->OnGetVolumePayloadData(PayloadKey, *UsdContext, Result);
+			}
 		}
 	}
 #endif	  // USE_USD_SDK
@@ -643,7 +665,10 @@ TOptional<UE::Interchange::FGroomPayloadData> UInterchangeUSDTranslator::GetGroo
 	{
 		for (const TSharedRef<FSchemaHandler>& Handler : ImplPtr->SchemaHandlers)
 		{
-			Handler->OnGetGroomPayloadData(PayloadKey, *UsdContext, Result);
+			if (Handler->IsEnabled())
+			{
+				Handler->OnGetGroomPayloadData(PayloadKey, *UsdContext, Result);
+			}
 		}
 	}
 #endif	  // USE_USD_SDK
@@ -666,7 +691,10 @@ TOptional<UE::Interchange::FInterchangeAudioPayloadData> UInterchangeUSDTranslat
 	{
 		for (const TSharedRef<FSchemaHandler>& Handler : ImplPtr->SchemaHandlers)
 		{
-			Handler->OnGetAudioPayloadData(PayloadKey, *UsdContext, Result);
+			if (Handler->IsEnabled())
+			{
+				Handler->OnGetAudioPayloadData(PayloadKey, *UsdContext, Result);
+			}
 		}
 	}
 #endif	  // USE_USD_SDK
diff --git a/Engine/Plugins/Interchange/Extensions/OpenUSD/Source/Import/Private/SchemaHandlers/MaterialSchemaHandler.cpp b/Engine/Plugins/Interchange/Extensions/OpenUSD/Source/Import/Private/SchemaHandlers/MaterialSchemaHandler.cpp
new file mode 100644
--- /dev/null
+++ b/Engine/Plugins/Interchange/Extensions/OpenUSD/Source/Import/Private/SchemaHandlers/MaterialSchemaHandler.cpp
@@ -0,0 +1,31 @@
+// Copyright Epic Games, Inc. All Rights Reserved.
+
+#if USE_USD_SDK
+
+#include "SchemaHandlers/MaterialSchemaHandler.h"
+
+#include "USDShadeConversion.h"
+#include "UsdWrappers/UsdPrim.h"
+
+namespace UE::Interchange::USD
+{
+	const FString& FMaterialSchemaHandler::GetTargetSchemaName() const
+	{
+		const static FString SchemaName = TEXT("Material");
+		return SchemaName;
+	}
+
+	bool FMaterialSchemaHandler::CanHandlePrim(const UE::FUsdPrim& Prim, const UInterchangeUsdContext& UsdContext) const
+	{
+		TRACE_CPUPROFILER_EVENT_SCOPE(FMaterialSchemaHandler::CanHandlePrim)
+
+		if (!Prim.IsA(*GetTargetSchemaName()))
+		{
+			return false;
+		}
+
+		return UsdUtils::HasSurfaceOutput(Prim, GetDefaultRenderContexts());
+	}
+}	 // namespace UE::Interchange::USD
+
+#endif	  // USE_USD_SDK
diff --git a/Engine/Plugins/Interchange/Extensions/OpenUSD/Source/Import/Private/SchemaHandlers/MaterialXSchemaHandler.cpp b/Engine/Plugins/Interchange/Extensions/OpenUSD/Source/Import/Private/SchemaHandlers/MaterialXSchemaHandler.cpp
--- a/Engine/Plugins/Interchange/Extensions/OpenUSD/Source/Import/Private/SchemaHandlers/MaterialXSchemaHandler.cpp
+++ b/Engine/Plugins/Interchange/Extensions/OpenUSD/Source/Import/Private/SchemaHandlers/MaterialXSchemaHandler.cpp
@@ -18,6 +18,7 @@
 
 #include "InterchangeManager.h"
 #include "InterchangeMaterialInstanceNode.h"
+#include "InterchangeMaterialReferenceNode.h"
 #include "InterchangeSceneNode.h"
 #include "InterchangeShaderGraphNode.h"
 #include "Nodes/InterchangeBaseNodeContainer.h"
@@ -76,9 +77,6 @@ namespace UE::MaterialXSchemaHandler::Private
 		TRACE_CPUPROFILER_EVENT_SCOPE(AddMaterialXShaderGraph)
 
 #if WITH_EDITOR
-		FName RenderContext = UsdContext.GetTranslatorSettings() ? UsdContext.GetTranslatorSettings()->RenderContext
-																  : UnrealIdentifiers::UnrealRenderContext;
-
 		TArray<FString> FilePaths = UsdUtils::GetMaterialXFilePaths(Prim);
 		for (const FString& File : FilePaths)
 		{
@@ -106,7 +104,7 @@ namespace UE::MaterialXSchemaHandler::Private
 
 		// Conditionally enable the shader graph for the time being as it crashes on Linux cause of a probable double free
 #if ENABLE_USD_MATERIALX
-		FUsdMaterialXShaderGraph ShaderGraph(Prim, *RenderContext.ToString());
+		FUsdMaterialXShaderGraph ShaderGraph(Prim, *UnrealIdentifiers::MaterialXRenderContext.ToString());
 		GeomPropValues = ShaderGraph.GetGeomPropValueNames();
 		if (MaterialX::DocumentPtr Document = ShaderGraph.GetDocument())
 		{
@@ -128,10 +126,10 @@ namespace UE::Interchange::USD
 		return HandlerName;
 	}
 
-	const FString& FMaterialXSchemaHandler::GetTargetSchemaName() const
+	const TArray<FString>& FMaterialXSchemaHandler::GetDefaultRenderContexts() const
 	{
-		const static FString SchemaName = TEXT("Material");
-		return SchemaName;
+		const static TArray<FString> RenderContexts{UnrealIdentifiers::MaterialXRenderContext.ToString()};
+		return RenderContexts;
 	}
 
 	bool FMaterialXSchemaHandler::CanHandlePrim(const UE::FUsdPrim& Prim, const UInterchangeUsdContext& UsdContext) const
@@ -139,25 +137,15 @@ namespace UE::Interchange::USD
 		TRACE_CPUPROFILER_EVENT_SCOPE(FMaterialXSchemaHandler::CanHandlePrim)
 
 #if WITH_EDITOR && ENABLE_USD_MATERIALX
-		if (!Prim.IsA(TEXT("Material")))
-		{
-			return false;
-		}
-
-		UInterchangeUsdTranslatorSettings* Settings = UsdContext.GetTranslatorSettings();
-		if (!Settings)
-		{
-			return false;
-		}
-
-		if (Settings->RenderContext != UnrealIdentifiers::MaterialXRenderContext)
+		if (!Prim.IsA(*GetTargetSchemaName()))
 		{
 			return false;
 		}
 
-		return UsdUtils::HasSurfaceOutput(Prim, UnrealIdentifiers::MaterialXRenderContext)
-			   || UsdUtils::HasVolumeOutput(Prim, UnrealIdentifiers::MaterialXRenderContext)
-			   || UsdUtils::HasDisplacementOutput(Prim, UnrealIdentifiers::MaterialXRenderContext);
+		const TArray<FString>& RenderContexts = GetDefaultRenderContexts();
+		return UsdUtils::HasSurfaceOutput(Prim, RenderContexts)
+			   || UsdUtils::HasVolumeOutput(Prim, RenderContexts)
+			   || UsdUtils::HasDisplacementOutput(Prim, RenderContexts);
 #else	 // #if WITH_EDITOR && ENABLE_USD_MATERIALX
 		return false;
 #endif
@@ -180,6 +168,25 @@ namespace UE::Interchange::USD
 			return {};
 		}
 
+		// For the material schema handlers, we only touch the asset if no other handler has produced a material node yet.
+		// This helps manage the separate material handlers per render context, without having them partially overwrite the
+		// node with render-context-specific data over and over
+		const static TSet<UClass*> MaterialLikeClasses = {
+			UInterchangeMaterialInstanceNode::StaticClass(),
+			UInterchangeShaderGraphNode::StaticClass(),
+			UInterchangeMaterialReferenceNode::StaticClass()
+		};
+		for (const UInterchangeBaseNode* AssetNode : AccumulatedInfo.PrimAssetNodes)
+		{
+			for (const UClass* MaterialClass : MaterialLikeClasses)
+			{
+				if (AssetNode->IsA(MaterialClass))
+				{
+					return false;
+				}
+			}
+		}
+
 		FString MaterialPrimName = Prim.GetName().ToString();
 
 		TArray<FUsdMaterialXShaderGraph::FGeomProp> GeomPropValues;
diff --git a/Engine/Plugins/Interchange/Extensions/OpenUSD/Source/Import/Private/SchemaHandlers/SchemaHandler.cpp b/Engine/Plugins/Interchange/Extensions/OpenUSD/Source/Import/Private/SchemaHandlers/SchemaHandler.cpp
--- a/Engine/Plugins/Interchange/Extensions/OpenUSD/Source/Import/Private/SchemaHandlers/SchemaHandler.cpp
+++ b/Engine/Plugins/Interchange/Extensions/OpenUSD/Source/Import/Private/SchemaHandlers/SchemaHandler.cpp
@@ -7,26 +7,25 @@
 #include "InterchangeUsdTraversalInfo.h"
 #include "SchemaHandlers/HandlerAccumulatedInfo.h"
 #include "SchemaHandlers/SchemaHandlerUtils.h"
-#include "USDLog.h"
-#include "UsdWrappers/SdfPath.h"
 
-#include "InterchangeSceneNode.h"
-#include "Nodes/InterchangeBaseNodeContainer.h"
-
-namespace UE::USDSchemaHandler::Private
-{
-	// 0 means invalid/unused
-	static UE::Interchange::USD::FSchemaHandlerIdentifier LastUsedIdentifier = 0;
-}
+#define LOCTEXT_NAMESPACE "UsdSchemaHandler"
 
 namespace UE::Interchange::USD
 {
-	TArray<FSchemaHandlerEntry> FSchemaHandlerRegistry::RegisteredHandlerEntries;
-
 	FSchemaHandler::~FSchemaHandler()
 	{
 	}
 
+	bool FSchemaHandler::IsEnabled() const
+	{
+		return bEnabled;
+	}
+
+	void FSchemaHandler::SetEnabled(bool bInEnabled)
+	{
+		bEnabled = bInEnabled;
+	}
+
 	const FString& FSchemaHandler::GetTargetSchemaName() const
 	{
 		const static FString Empty;
@@ -143,60 +142,8 @@ namespace UE::Interchange::USD
 
 		return false;
 	}
-
-	FSchemaHandlerEntry::FSchemaHandlerEntry(const FString& InHandlerName, const FString& InSchemaName, FSchemaHandlerGenerator InHandlerGenerator)
-		: Identifier(++UE::USDSchemaHandler::Private::LastUsedIdentifier)
-		, HandlerName(InHandlerName)
-		, SchemaName(InSchemaName)
-		, HandlerGenerator(InHandlerGenerator)
-	{
-	}
-
-	FSchemaHandlerIdentifier FSchemaHandlerEntry::GetIdentifier() const
-	{
-		return Identifier;
-	}
-
-	const FString& FSchemaHandlerEntry::GetHandlerName() const
-	{
-		return HandlerName;
-	}
-
-	const FString& FSchemaHandlerEntry::GetTargetSchemaName() const
-	{
-		return SchemaName;
-	}
-
-	FSchemaHandlerGenerator FSchemaHandlerEntry::GetHandlerGenerator() const
-	{
-		return HandlerGenerator;
-	}
-
-	bool FSchemaHandlerRegistry::Unregister(FSchemaHandlerIdentifier Identifier)
-	{
-		uint64 NumRemoved = RegisteredHandlerEntries.RemoveAll(
-			[Identifier](const FSchemaHandlerEntry& Rec)
-			{
-				return Rec.GetIdentifier() == Identifier;
-			}
-		);
-		return NumRemoved > 0;
-	}
-
-	TArray<TSharedRef<FSchemaHandler>> FSchemaHandlerRegistry::GenerateHandlers()
-	{
-		TArray<TSharedRef<FSchemaHandler>> Result;
-		Result.Reserve(RegisteredHandlerEntries.Num());
-
-		for (const FSchemaHandlerEntry& Entry : RegisteredHandlerEntries)
-		{
-			FSchemaHandlerGenerator Generator = Entry.GetHandlerGenerator();
-			TSharedRef<FSchemaHandler> NewHandler = Generator();
-			Result.Add(NewHandler);
-		}
-
-		return Result;
-	};
 }	 // namespace UE::Interchange::USD
 
+#undef LOCTEXT_NAMESPACE
+
 #endif	  // USE_USD_SDK
diff --git a/Engine/Plugins/Interchange/Extensions/OpenUSD/Source/Import/Private/SchemaHandlers/SchemaHandlerRegistry.cpp b/Engine/Plugins/Interchange/Extensions/OpenUSD/Source/Import/Private/SchemaHandlers/SchemaHandlerRegistry.cpp
new file mode 100644
--- /dev/null
+++ b/Engine/Plugins/Interchange/Extensions/OpenUSD/Source/Import/Private/SchemaHandlers/SchemaHandlerRegistry.cpp
@@ -0,0 +1,91 @@
+// Copyright Epic Games, Inc. All Rights Reserved.
+
+#if USE_USD_SDK
+
+#include "SchemaHandlers/SchemaHandlerRegistry.h"
+#include "USDErrorUtils.h"
+
+#define LOCTEXT_NAMESPACE "InterchangeUSDTranslator"
+
+namespace UE::Interchange::USD
+{
+	TArray<FSchemaHandlerEntry> FSchemaHandlerRegistry::RegisteredHandlerEntries;
+	TMap<FString, FSchemaHandlerGenerator> FSchemaHandlerRegistry::RegisteredHandlerNamesToGenerators;
+
+	bool FSchemaHandlerRegistry::Unregister(const FString& Name)
+	{
+		uint64 NumRemoved = RegisteredHandlerEntries.RemoveAll(
+			[Name](const FSchemaHandlerEntry& Rec)
+			{
+				return Rec.HandlerName == Name;
+			}
+		);
+
+		uint64 NumGeneratorRemoved = RegisteredHandlerNamesToGenerators.Remove(Name);
+
+		// We will validate duplicate entries on FSchemaHandlerRegistry::GenerateHandlers, in here a success
+		// means having removed *anything*. This return value is not particularly checked by anything at this
+		// point anyway...
+		return NumRemoved > 0 || NumGeneratorRemoved > 0;
+	}
+
+	TArray<TSharedRef<FSchemaHandler>> FSchemaHandlerRegistry::GenerateHandlers(const TArray<FSchemaHandlerEntry>* Entries)
+	{
+		if (Entries == nullptr)
+		{
+			return {};
+		}
+
+		const uint64 NumEntries = Entries->Num();
+
+		TArray<TSharedRef<FSchemaHandler>> Result;
+		Result.Reserve(NumEntries);
+
+		TSet<FString> SeenHandlerNames;
+		SeenHandlerNames.Reserve(NumEntries);
+
+		for (const FSchemaHandlerEntry& Entry : *Entries)
+		{
+			FSchemaHandlerGenerator* FoundGenerator = RegisteredHandlerNamesToGenerators.Find(Entry.HandlerName);
+			if (!FoundGenerator)
+			{
+				USD_LOG_USERERROR(FText::Format(
+					LOCTEXT("MissingHandlerGenerator", "Registered schema handler '{0}' does not have a generator function assigned."),
+					FText::FromString(Entry.HandlerName)
+				));
+				continue;
+			}
+
+			// Apply (potentially modified) settings stored on the entries back to the handlers
+			TSharedRef<FSchemaHandler> NewHandler = (*FoundGenerator)();
+			NewHandler->SetEnabled(Entry.bEnabled);
+
+			// We check for duplicate handler names here instead of in Registry because it's possible that user C++ code may have manipulated 
+			// the registered handler array directly, adding/removing handlers (as it was intended) and skipping our own Register function.
+			// This GenerateHandlers however is called by the USD Translator directly and is not affected by user code, so we know
+			// we can validate this fully at this point.
+			//
+			// Also, the fact that this is normally called during translation means we should have the Interchange result container in place, and
+			// the USD error utils should capture this error message and show the toast mentioning how there were issues with the import
+			if (SeenHandlerNames.Contains(Entry.HandlerName))
+			{
+				USD_LOG_USERERROR(FText::Format(
+					LOCTEXT(
+						"DuplicateHandlerName",
+						"Registered schema handlers should have unique names, but name '{0}' is used more than once. Translation may be inaccurate."
+					),
+					FText::FromString(Entry.HandlerName)
+				));
+			}
+			SeenHandlerNames.Add(Entry.HandlerName);
+
+			Result.Add(NewHandler);
+		}
+
+		return Result;
+	};
+}	 // namespace UE::Interchange::USD
+
+#undef LOCTEXT_NAMESPACE
+
+#endif	  // USE_USD_SDK
diff --git a/Engine/Plugins/Interchange/Extensions/OpenUSD/Source/Import/Private/SchemaHandlers/SchemaHandlerUtils.cpp b/Engine/Plugins/Interchange/Extensions/OpenUSD/Source/Import/Private/SchemaHandlers/SchemaHandlerUtils.cpp
--- a/Engine/Plugins/Interchange/Extensions/OpenUSD/Source/Import/Private/SchemaHandlers/SchemaHandlerUtils.cpp
+++ b/Engine/Plugins/Interchange/Extensions/OpenUSD/Source/Import/Private/SchemaHandlers/SchemaHandlerUtils.cpp
@@ -7,6 +7,7 @@
 #include "InterchangeUsdContext.h"
 #include "InterchangeUsdTranslator.h"
 #include "SchemaHandlers/SchemaHandler.h"
+#include "SchemaHandlers/SchemaHandlerEntry.h"
 #include "Usd/InterchangeUsdDefinitions.h"
 #include "USDErrorUtils.h"
 #include "USDNaniteAssemblyUtils.h"
diff --git a/Engine/Plugins/Interchange/Extensions/OpenUSD/Source/Import/Private/SchemaHandlers/UniversalMaterialSchemaHandler.cpp b/Engine/Plugins/Interchange/Extensions/OpenUSD/Source/Import/Private/SchemaHandlers/UniversalMaterialSchemaHandler.cpp
--- a/Engine/Plugins/Interchange/Extensions/OpenUSD/Source/Import/Private/SchemaHandlers/UniversalMaterialSchemaHandler.cpp
+++ b/Engine/Plugins/Interchange/Extensions/OpenUSD/Source/Import/Private/SchemaHandlers/UniversalMaterialSchemaHandler.cpp
@@ -325,10 +325,12 @@ namespace UE::Interchange::USD
 		return HandlerName;
 	}
 
-	const FString& FUniversalMaterialSchemaHandler::GetTargetSchemaName() const
+	const TArray<FString>& FUniversalMaterialSchemaHandler::GetDefaultRenderContexts() const
 	{
-		const static FString SchemaName = TEXT("Material");
-		return SchemaName;
+		// The empty string corresponds to the 'universal' render context. When displaying on the UI, we'll swap this empty
+		// string to/from the actual 'universal' string
+		const static TArray<FString> RenderContexts{TEXT("")};
+		return RenderContexts;
 	}
 
 	TOptional<bool> FUniversalMaterialSchemaHandler::CanBeCollapsed(const UE::FUsdPrim& Prim, UInterchangeUsdContext& UsdContext) const
@@ -358,10 +360,9 @@ namespace UE::Interchange::USD
 			return {};
 		}
 
-		// This handler works slightly differently: We'll only emit an asset node if we've found no other material-like asset node
-		// already produced for this prim. The idea is to act as a fallback material handler
-		//
-		// TODO: Add common base class to all material-like nodes?
+		// For the material schema handlers, we only touch the asset if no other handler has produced a material node yet.
+		// This helps manage the separate material handlers per render context, without having them partially overwrite the
+		// node with render-context-specific data over and over
 		const static TSet<UClass*> MaterialLikeClasses = {
 			UInterchangeMaterialInstanceNode::StaticClass(),
 			UInterchangeShaderGraphNode::StaticClass(),
diff --git a/Engine/Plugins/Interchange/Extensions/OpenUSD/Source/Import/Private/SchemaHandlers/UnrealMaterialSchemaHandler.cpp b/Engine/Plugins/Interchange/Extensions/OpenUSD/Source/Import/Private/SchemaHandlers/UnrealMaterialSchemaHandler.cpp
--- a/Engine/Plugins/Interchange/Extensions/OpenUSD/Source/Import/Private/SchemaHandlers/UnrealMaterialSchemaHandler.cpp
+++ b/Engine/Plugins/Interchange/Extensions/OpenUSD/Source/Import/Private/SchemaHandlers/UnrealMaterialSchemaHandler.cpp
@@ -13,16 +13,12 @@
 #include "USDConversionUtils.h"
 #include "USDShadeConversion.h"
 #include "UsdWrappers/UsdPrim.h"
-#include "UsdWrappers/UsdStage.h"
 
 #include "InterchangeMaterialInstanceNode.h"
 #include "InterchangeMaterialReferenceNode.h"
-#include "InterchangeSceneNode.h"
 #include "InterchangeShaderGraphNode.h"
 #include "Nodes/InterchangeBaseNodeContainer.h"
 
-#include "UDIMUtilities.h"
-
 namespace UE::Interchange::USD
 {
 	const FString& FUnrealMaterialSchemaHandler::GetHandlerName() const
@@ -31,28 +27,10 @@ namespace UE::Interchange::USD
 		return HandlerName;
 	}
 
-	const FString& FUnrealMaterialSchemaHandler::GetTargetSchemaName() const
-	{
-		const static FString SchemaName = TEXT("Material");
-		return SchemaName;
-	}
-
-	bool FUnrealMaterialSchemaHandler::CanHandlePrim(const UE::FUsdPrim& Prim, const UInterchangeUsdContext& UsdContext) const
+	const TArray<FString>& FUnrealMaterialSchemaHandler::GetDefaultRenderContexts() const
 	{
-		TRACE_CPUPROFILER_EVENT_SCOPE(FUnrealMaterialSchemaHandler::CanHandlePrim)
-
-		if (!Prim.IsA(TEXT("Material")))
-		{
-			return false;
-		}
-
-		UInterchangeUsdTranslatorSettings* Settings = UsdContext.GetTranslatorSettings();
-		if (!Settings || Settings->RenderContext != UnrealIdentifiers::UnrealRenderContext)
-		{
-			return false;
-		}
-
-		return UsdUtils::HasSurfaceOutput(Prim, UnrealIdentifiers::UnrealRenderContext);
+		const static TArray<FString> RenderContexts{*UnrealIdentifiers::UnrealRenderContext.ToString()};
+		return RenderContexts;
 	}
 
 	bool FUnrealMaterialSchemaHandler::OnTranslate(
@@ -70,6 +48,25 @@ namespace UE::Interchange::USD
 			return {};
 		}
 
+		// For the material schema handlers, we only touch the asset if no other handler has produced a material node yet.
+		// This helps manage the separate material handlers per render context, without having them partially overwrite the
+		// node with render-context-specific data over and over
+		const static TSet<UClass*> MaterialLikeClasses = {
+			UInterchangeMaterialInstanceNode::StaticClass(),
+			UInterchangeShaderGraphNode::StaticClass(),
+			UInterchangeMaterialReferenceNode::StaticClass()
+		};
+		for (const UInterchangeBaseNode* AssetNode : AccumulatedInfo.PrimAssetNodes)
+		{
+			for (const UClass* MaterialClass : MaterialLikeClasses)
+			{
+				if (AssetNode->IsA(MaterialClass))
+				{
+					return false;
+				}
+			}
+		}
+
 		const FString PrimPath = UsdUtils::GetPrototypePrimPath(Prim).GetString();
 		FString MaterialPrimName = Prim.GetName().ToString();
 		FString NewNodeUid = MaterialPrefix + PrimPath;
diff --git a/Engine/Plugins/Interchange/Extensions/OpenUSD/Source/Import/Private/SchemaHandlers/VolumeSchemaHandler.cpp b/Engine/Plugins/Interchange/Extensions/OpenUSD/Source/Import/Private/SchemaHandlers/VolumeSchemaHandler.cpp
--- a/Engine/Plugins/Interchange/Extensions/OpenUSD/Source/Import/Private/SchemaHandlers/VolumeSchemaHandler.cpp
+++ b/Engine/Plugins/Interchange/Extensions/OpenUSD/Source/Import/Private/SchemaHandlers/VolumeSchemaHandler.cpp
@@ -7,18 +7,21 @@
 #include "InterchangeUsdContext.h"
 #include "InterchangeUsdTranslator.h"
 #include "InterchangeUsdTraversalInfo.h"
+#include "SchemaHandlers/SchemaHandlerRegistry.h"
 #include "SchemaHandlers/SchemaHandlerUtils.h"
 #include "Usd/InterchangeUsdDefinitions.h"
 #include "USDErrorUtils.h"
 #include "USDPrimConversion.h"
+#include "USDShadeConversion.h"
 #include "USDTypesConversion.h"
 #include "UsdWrappers/UsdAttribute.h"
 #include "UsdWrappers/UsdPrim.h"
 
-#include "InterchangeManager.h"
 #include "InterchangeCommonAnimationPayload.h"
+#include "InterchangeManager.h"
 #include "InterchangeMaterialInstanceNode.h"
 #include "InterchangeSceneNode.h"
+#include "InterchangeShaderGraphNode.h"
 #include "InterchangeSourceData.h"
 #include "InterchangeVolumeNode.h"
 #include "Nodes/InterchangeBaseNodeContainer.h"
@@ -30,6 +33,10 @@
 #include "Async/ParallelFor.h"
 #include "Components/HeterogeneousVolumeComponent.h"
 
+#include "USDIncludesStart.h"
+#include "pxr/base/tf/token.h"
+#include "USDIncludesEnd.h"
+
 #define LOCTEXT_NAMESPACE "VolumeSchemaHandler"
 
 namespace UE::VolVolumeSchemaHandler::Private
@@ -236,6 +243,81 @@ namespace UE::VolVolumeSchemaHandler::Private
 		return MaterialInstance;
 	}
 
+	UInterchangeShaderGraphNode* TryCreatingVolumetricShaderGraphNode(const UE::FUsdPrim& VolumePrim, UInterchangeUsdContext& UsdContext)
+	{
+		const bool bProvideMaterialIndices = false;
+		const pxr::TfToken MtlxToken{"mtlx"};	 // We no longer need to pass this, but there's no reason *not* to pass mtlx here either
+		UsdUtils::FUsdPrimMaterialAssignmentInfo Assignments = UsdUtils::GetPrimMaterialAssignments(
+			VolumePrim,
+			UsdUtils::GetDefaultTimeCode(),
+			bProvideMaterialIndices,
+			MtlxToken,
+			UsdContext.CachedMeshConversionOptions.MaterialPurpose
+		);
+
+		for (const UsdUtils::FUsdPrimMaterialSlot& Slot : Assignments.Slots)
+		{
+			if (Slot.AssignmentType != UsdUtils::EPrimAssignmentType::MaterialPrim)
+			{
+				continue;
+			}
+
+			// MaterialSource here is the material prim path
+			UE::FUsdPrim MaterialPrim = UsdContext.GetUsdStage().GetPrimAtPath(UE::FSdfPath{*Slot.MaterialSource});
+			const bool bMaterialXOutput = UsdUtils::HasSurfaceOutput(MaterialPrim, UnrealIdentifiers::MaterialXRenderContext)
+										  || UsdUtils::HasVolumeOutput(MaterialPrim, UnrealIdentifiers::MaterialXRenderContext)
+										  || UsdUtils::HasDisplacementOutput(MaterialPrim, UnrealIdentifiers::MaterialXRenderContext);
+			if (!bMaterialXOutput)
+			{
+				continue;
+			}
+
+			UInterchangeUsdTranslatorSettings* Settings = UsdContext.GetTranslatorSettings();
+			if (!Settings)
+			{
+				continue;
+			}
+
+			// Check if we have a handler capable of parsing MaterialX
+			// This is not a perfect check, but it should hopefully be good enough for now
+			bool bMaterialXHandler = false;
+			const TArray<FSchemaHandlerEntry>* HandlerEntries = &FSchemaHandlerRegistry::RegisteredHandlerEntries;
+			if (Settings && Settings->CustomHandlerEntries.Num() > 0)
+			{
+				HandlerEntries = &Settings->CustomHandlerEntries;
+			}
+			if (HandlerEntries)
+			{
+				for (const FSchemaHandlerEntry& Entry : *HandlerEntries)
+				{
+					const static FString MtlxString = *UnrealIdentifiers::MaterialXRenderContext.ToString();
+					if (Entry.bEnabled && Entry.DefaultRenderContexts.Contains(MtlxString))
+					{
+						bMaterialXHandler = true;
+					}
+				}
+			}
+			if (!bMaterialXHandler)
+			{
+				continue;
+			}
+
+			if (UInterchangeUSDTranslator* CurrentTranslator = UsdContext.GetTranslator())
+			{
+				FTraversalInfo TraversalInfo;
+				if (TOptional<FHandlerAccumulatedInfo> AccumulatedInfoForPrim = CurrentTranslator->TranslatePrim(MaterialPrim, TraversalInfo))
+				{
+					if (UInterchangeShaderGraphNode* MaterialNode = AccumulatedInfoForPrim->GetAssetNodeOfClass<UInterchangeShaderGraphNode>())
+					{
+						return MaterialNode;
+					}
+				}
+			}
+		}
+
+		return nullptr;
+	}
+
 	TSet<UInterchangeVolumeNode*> CreateVolumeNodes(
 		const UE::FUsdPrim& Prim,
 		TMap<FString, TArray<UsdUtils::FVolumePrimInfo>>& InOutPrimPathToVolumeInfo,
@@ -461,73 +543,20 @@ namespace UE::VolVolumeSchemaHandler::Private
 			}
 		}
 
-		// In case of MaterialX context, the volume material will be created from a shader graph
-		UInterchangeUsdTranslatorSettings* TranslatorSettings = UsdContext.GetTranslatorSettings();
-		if (TranslatorSettings && (TranslatorSettings->RenderContext == UnrealIdentifiers::MaterialXRenderContext))
+		// If we have a material with a MaterialX output and a MaterialX handler is enabled, then let's assume we'll end up parsing
+		// this material as a MaterialX shader graph and return that
+		if (UInterchangeShaderGraphNode* ShaderGraphMaterial = TryCreatingVolumetricShaderGraphNode(Prim, UsdContext))
 		{
-			const bool bProvideMaterialIndices = false;
-			UsdUtils::FUsdPrimMaterialAssignmentInfo Assignments = UsdUtils::GetPrimMaterialAssignments(
-				Prim,
-				UsdUtils::GetDefaultTimeCode(),
-				bProvideMaterialIndices,
-				UsdContext.CachedMeshConversionOptions.RenderContext,
-				UsdContext.CachedMeshConversionOptions.MaterialPurpose
-			);
-
-			for (const UsdUtils::FUsdPrimMaterialSlot& Slot : Assignments.Slots)
-			{
-				FString SlotName = Slot.MaterialSource;
-
-				if (Slot.AssignmentType != UsdUtils::EPrimAssignmentType::MaterialPrim)
-				{
-					continue;
-				}
-
-				// MaterialSource here is the material prim path
-				UE::FUsdPrim MaterialPrim = UsdContext.GetUsdStage().GetPrimAtPath(UE::FSdfPath{ *Slot.MaterialSource });
-
-				// Go through a prototype in case this material is inside an instance proxy, as this needs to match up with
-				// the polygon group name we're going to add to our mesh description.
-				UE::FUsdPrim MaterialPrototypePrim = UsdUtils::GetPrototypePrim(MaterialPrim);
-				if (MaterialPrototypePrim != MaterialPrim)
-				{
-					SlotName = MaterialPrototypePrim.GetPrimPath().GetString();
-				}
-
-				if (UInterchangeUSDTranslator* CurrentTranslator = UsdContext.GetTranslator())
-				{
-					FTraversalInfo TraversalInfo;
-					if (TOptional<FHandlerAccumulatedInfo> AccumulatedInfoForPrim = CurrentTranslator->TranslatePrim(MaterialPrototypePrim, TraversalInfo))
-					{
-						// The material node types don't share a base class...
-						UInterchangeBaseNode* MaterialNode = AccumulatedInfoForPrim->GetAssetNodeOfClass<UInterchangeBaseNode>();
-
-						OutMaterialInstanceUid = MaterialNode->GetUniqueID();
-
-						if (Slot.bMeshIsDoubleSided)
-						{
-							if (UInterchangeBaseNode* TwoSidedMaterial = GetOrCreateTwoSidedMaterial(MaterialNode, *UsdContext.GetNodeContainer()))
-							{
-								OutMaterialInstanceUid = TwoSidedMaterial->GetUniqueID();
-							}
-						}
-					}
-				}
-
-				if (!OutMaterialInstanceUid.IsEmpty())
-				{
-					break;
-				}
-			}
-
-			if (!OutMaterialInstanceUid.IsEmpty())
-			{
-				return VolumeNodes;
-			}
+			OutMaterialInstanceUid = ShaderGraphMaterial->GetUniqueID();
+			return VolumeNodes;
 		}
+
 		// Setup a new material instance node for this volume, which will be used by the Heterogeneous Volume Actor we'll also
 		// spawn for this volume prim.
 		//
+		// Note that this is slightly different from other material handling because in here we need a material *instance*, even if
+		// we'se using Unreal materials, so we don't use our regular SetSlotMaterialDependencies() call
+		//
 		// Note that there's only one material slot per actor, but that we really do need some kind of material to be put in there,
 		// so that we can at least assign our SVTs somewhere.
 		UInterchangeMaterialInstanceNode* MaterialInstance = CreateVolumetricMaterialInstanceNode(
diff --git a/Engine/Plugins/Interchange/Extensions/OpenUSD/Source/Import/Public/InterchangeOpenUSDImportModule.h b/Engine/Plugins/Interchange/Extensions/OpenUSD/Source/Import/Public/InterchangeOpenUSDImportModule.h
--- a/Engine/Plugins/Interchange/Extensions/OpenUSD/Source/Import/Public/InterchangeOpenUSDImportModule.h
+++ b/Engine/Plugins/Interchange/Extensions/OpenUSD/Source/Import/Public/InterchangeOpenUSDImportModule.h
@@ -19,6 +19,6 @@ public:
 #if USE_USD_SDK
 private:
 	FDelegateHandle PostEngineInitHandle;
-	TArray<UE::Interchange::USD::FSchemaHandlerIdentifier> RegisteredHandlers;
+	TArray<FString> RegisteredHandlers;
 #endif	  // USE_USD_SDK
 };
diff --git a/Engine/Plugins/Interchange/Extensions/OpenUSD/Source/Import/Public/InterchangeUsdTranslator.h b/Engine/Plugins/Interchange/Extensions/OpenUSD/Source/Import/Public/InterchangeUsdTranslator.h
--- a/Engine/Plugins/Interchange/Extensions/OpenUSD/Source/Import/Public/InterchangeUsdTranslator.h
+++ b/Engine/Plugins/Interchange/Extensions/OpenUSD/Source/Import/Public/InterchangeUsdTranslator.h
@@ -12,6 +12,7 @@
 #include "Volume/InterchangeVolumePayloadInterface.h"
 
 #include "InterchangeUsdTraversalInfo.h"
+#include "SchemaHandlers/SchemaHandlerEntry.h"
 #include "USDStageOptions.h"
 
 #include "UObject/Object.h"
@@ -61,14 +62,14 @@ class UInterchangeUsdTranslatorSettings : public UInterchangeTranslatorSettings
 	GENERATED_BODY()
 
 public:
-
 	/** Only import geometry prims with these specific purposes from the USD file */
 	UE_DEPRECATED(5.6, "This option is no longer a translator setting and has moved to the InterchangeUSDPipeline")
 	UPROPERTY()
 	int32 GeometryPurpose;
 
 	/** Specifies which set of shaders to use when parsing USD materials, in addition to the universal render context. */
-	UPROPERTY(EditAnywhere, Category = "USD Translator")
+	UE_DEPRECATED(5.6, "This is no longer used: Reorder and enable/disable the registered material schema handlers to control what render contexts are parsed (see the CustomHandlerEntries property).")
+	UPROPERTY()
 	FName RenderContext;
 
 	/** Specifies which material purpose to use when parsing USD material bindings, in addition to the "allPurpose" fallback */
@@ -138,6 +139,20 @@ public:
 	UPROPERTY(EditAnywhere, Category = "USD Translator", meta = (EditCondition = bTranslatePrimMetadata))
 	FString MetadataRegexFilter;
 
+	/**
+	 * Custom array of schema handlers to use for translation.
+	 * When this is empty, the default FSchemaHandlerRegistry::RegisteredHandlerEntries array will be used instead.
+	 *
+	 * It is expected for this array to essentially be a reorder of FSchemaHandlerRegistry::RegisteredHandlerEntries, with additionally
+	 * some modified property values on schema handler entries, if needed (e.g. enabled or custom render contexts).
+	 *
+	 * It is not sufficient to create a new FSchemaHandlerEntry instance from scratch and insert here directly: If you have a custom handler,
+	 * you want to register it on FSchemaHandlerRegistry instead, so as to also get the handler generator function correctly registered.
+	 * At that point, your handler will also be added to FSchemaHandlerRegistry::RegisteredHandlerEntries.
+	 */
+	UPROPERTY(EditAnywhere, Category = "USD Translator|Handlers", AdvancedDisplay)
+	TArray<FSchemaHandlerEntry> CustomHandlerEntries;
+
 public:
 	UInterchangeUsdTranslatorSettings();
 };
diff --git a/Engine/Plugins/Interchange/Extensions/OpenUSD/Source/Import/Public/SchemaHandlers/MaterialSchemaHandler.h b/Engine/Plugins/Interchange/Extensions/OpenUSD/Source/Import/Public/SchemaHandlers/MaterialSchemaHandler.h
new file mode 100644
--- /dev/null
+++ b/Engine/Plugins/Interchange/Extensions/OpenUSD/Source/Import/Public/SchemaHandlers/MaterialSchemaHandler.h
@@ -0,0 +1,30 @@
+// Copyright Epic Games, Inc. All Rights Reserved.
+
+#pragma once
+
+#if USE_USD_SDK
+
+#include "SchemaHandler.h"
+
+#define UE_API INTERCHANGEOPENUSDIMPORT_API
+
+namespace UE::Interchange::USD
+{
+	class FMaterialSchemaHandler : public FSchemaHandler
+	{
+	public:
+		UE_API virtual const FString& GetTargetSchemaName() const override;
+
+		/**
+		 * Returns the list of render contexts that this handler supports by default.
+		 * This is likely to be a single value, such as 'unreal' or 'mtlx'
+		 */
+		UE_API virtual const TArray<FString>& GetDefaultRenderContexts() const = 0;
+
+		UE_API virtual bool CanHandlePrim(const UE::FUsdPrim& Prim, const UInterchangeUsdContext& UsdContext) const override;
+	};
+}
+
+#undef UE_API
+
+#endif	  // USE_USD_SDK
\ No newline at end of file
diff --git a/Engine/Plugins/Interchange/Extensions/OpenUSD/Source/Import/Public/SchemaHandlers/MaterialXSchemaHandler.h b/Engine/Plugins/Interchange/Extensions/OpenUSD/Source/Import/Public/SchemaHandlers/MaterialXSchemaHandler.h
--- a/Engine/Plugins/Interchange/Extensions/OpenUSD/Source/Import/Public/SchemaHandlers/MaterialXSchemaHandler.h
+++ b/Engine/Plugins/Interchange/Extensions/OpenUSD/Source/Import/Public/SchemaHandlers/MaterialXSchemaHandler.h
@@ -4,18 +4,18 @@
 
 #if USE_USD_SDK
 
-#include "SchemaHandler.h"
+#include "SchemaHandlers/MaterialSchemaHandler.h"
 
 #define UE_API INTERCHANGEOPENUSDIMPORT_API
 
 namespace UE::Interchange::USD
 {
-	class FMaterialXSchemaHandler : public FSchemaHandler
+	class FMaterialXSchemaHandler : public FMaterialSchemaHandler
 	{
 	public:
 		UE_API virtual const FString& GetHandlerName() const override;
 
-		UE_API virtual const FString& GetTargetSchemaName() const override;
+		UE_API virtual const TArray<FString>& GetDefaultRenderContexts() const override;
 
 		UE_API virtual bool CanHandlePrim(const UE::FUsdPrim& Prim, const UInterchangeUsdContext& UsdContext) const override;
 
diff --git a/Engine/Plugins/Interchange/Extensions/OpenUSD/Source/Import/Public/SchemaHandlers/SchemaHandler.h b/Engine/Plugins/Interchange/Extensions/OpenUSD/Source/Import/Public/SchemaHandlers/SchemaHandler.h
--- a/Engine/Plugins/Interchange/Extensions/OpenUSD/Source/Import/Public/SchemaHandlers/SchemaHandler.h
+++ b/Engine/Plugins/Interchange/Extensions/OpenUSD/Source/Import/Public/SchemaHandlers/SchemaHandler.h
@@ -9,6 +9,7 @@
 #include "Texture/InterchangeBlockedTexturePayloadData.h"
 #include "Texture/InterchangeTexturePayloadData.h"
 
+#include "SchemaHandlers/SchemaHandlerEntry.h"
 #include "UsdWrappers/UsdPrim.h"
 
 #include "Containers/AnsiString.h"
@@ -48,8 +49,6 @@ namespace UE
 
 namespace UE::Interchange::USD
 {
-	using FSchemaHandlerIdentifier = int32;
-
 	/**
 	 * Class that can be used to encapsulate the translation of prims of a certain type into Interchange nodes.
 	 *
@@ -63,16 +62,22 @@ namespace UE::Interchange::USD
 	 * accumulated Interchange nodes created for that prim.
 	 *
 	 * Whenever you don't want your handler to be invoked anymore (e.g. when your module is being shut down) you can unregister
-	 * your handler with FSchemaHandlerRegistry::Unregister(), passing the identifier that you received when calling Register().
+	 * your handler with FSchemaHandlerRegistry::Unregister(), passing the handler name
 	 */
 	class FSchemaHandler
 	{
 	public:
 		UE_API virtual ~FSchemaHandler();
 
-		/** Returns the name of this handler. This serves no practical purpose other than to help identifying and manipulating registered handlers */
+		/** Returns the name of this handler, which must be unique between all registered handlers */
 		virtual const FString& GetHandlerName() const = 0;
 
+		/** Gets whether this handler is enabled for this translation or not */
+		UE_API virtual bool IsEnabled() const;
+
+		/** Sets whether this handler is enabled for this translation or not */
+		UE_API virtual void SetEnabled(bool bEnabled);
+
 		/**
 		 * Returns the schema name of the prim type to which this handler applies ("Mesh", or "Material, or "Xform", etc.).
 		 *
@@ -176,86 +181,12 @@ namespace UE::Interchange::USD
 			UInterchangeUsdContext& UsdContext,
 			TOptional<UE::Interchange::FInterchangeAudioPayloadData>& InOutPayloadData
 		);
-	};
-
-	/** Main lambda invoked when spawning an FSchemaHandler */
-	using FSchemaHandlerGenerator = TFunction<TSharedRef<FSchemaHandler>()>;
-
-	/**
-	 * Represents a registered schema handler of a particular type.
-	 *
-	 * This class is used instead of directly tracking FSchemaHandlerGenerators in order to allow programmatically inspecting
-	 * and manipulating the registered handlers.
-	 */
-	class FSchemaHandlerEntry
-	{
-	public:
-		UE_API FSchemaHandlerEntry(const FString& HandlerName, const FString& SchemaName, FSchemaHandlerGenerator HandlerGenerator);
-
-		UE_API FSchemaHandlerIdentifier GetIdentifier() const;
-		UE_API const FString& GetHandlerName() const;
-		UE_API const FString& GetTargetSchemaName() const;
-		UE_API FSchemaHandlerGenerator GetHandlerGenerator() const;
-
-	private:
-		FSchemaHandlerIdentifier Identifier;
-		FString HandlerName;
-		FString SchemaName;
-		FSchemaHandlerGenerator HandlerGenerator;
-	};
-
-	/**
-	 * Static storage of all registered FSchemaHandlers.
-	 */
-	class FSchemaHandlerRegistry
-	{
-	public:
-		/**
-		 * Variable that stores the actual registered schema handlers.
-		 * This is intentionally exposed so as to allow free insertion/reorder/removal of registered handlers.
-		 */
-		static UE_API TArray<FSchemaHandlerEntry> RegisteredHandlerEntries;
-
-	public:
-		/**
-		 * Registers a schema handler of a particular type, so that it is invoked when translating USD files.
-		 * Usage:
-		 *     FSchemaHandlerIdentifier Id = FSchemaHandlerRegistry::Register<FImageableSchemaHandler>();
-		 */
-		template<typename SchemaHandlerType>
-		static FSchemaHandlerIdentifier Register()
-		{
-			SchemaHandlerType TempHandler{};
-
-			const FString& HandlerName = TempHandler.GetHandlerName();
-			const FString& SchemaName = TempHandler.GetTargetSchemaName();
-			FSchemaHandlerGenerator SpawnerFunc = []()
-			{
-				TSharedRef<FSchemaHandler> Result = MakeShared<SchemaHandlerType>();
-				return Result;
-			};
 
-			FSchemaHandlerEntry NewEntry{HandlerName, SchemaName, SpawnerFunc};
-			RegisteredHandlerEntries.Add(NewEntry);
-			return NewEntry.GetIdentifier();
-		}
-
-		/**
-		 * Unregisters a schema handler, so that it will no longer be invoked when translating USD files.
-		 */
-		static UE_API bool Unregister(FSchemaHandlerIdentifier Identifier);
-
-		/**
-		 * Spawns an instance of each entry in RegisteredHandlerEntries, in order.
-		 *
-		 * Before translation, the UInterchangeUSDTranslator calls this function and retains those same spawned handlers 
-		 * throughout the entire import process, so they are allowed to keep internal state between translation and 
-		 * payload retrieval
-		 */
-		static UE_API TArray<TSharedRef<FSchemaHandler>> GenerateHandlers();
+	protected:
+		bool bEnabled = true;
 	};
-}
+}	 // namespace UE::Interchange::USD
 
 #undef UE_API
 
-#endif	  // USE_USD_SDK
\ No newline at end of file
+#endif	  // USE_USD_SDK
diff --git a/Engine/Plugins/Interchange/Extensions/OpenUSD/Source/Import/Public/SchemaHandlers/SchemaHandlerEntry.h b/Engine/Plugins/Interchange/Extensions/OpenUSD/Source/Import/Public/SchemaHandlers/SchemaHandlerEntry.h
new file mode 100644
--- /dev/null
+++ b/Engine/Plugins/Interchange/Extensions/OpenUSD/Source/Import/Public/SchemaHandlers/SchemaHandlerEntry.h
@@ -0,0 +1,34 @@
+// Copyright Epic Games, Inc. All Rights Reserved.
+
+#pragma once
+
+#include "SchemaHandlerEntry.generated.h"
+
+/**
+ * Represents a registered schema handler of a particular type.
+ *
+ * This class is used instead of directly tracking FSchemaHandlerGenerators in order to allow programmatically inspecting
+ * and manipulating the registered persistent handler information without repeatedly instantiating handlers.
+ *
+ * It's a USTRUCT and kept in a separate file as it's used by UInterchangeUsdTranslatorSettings in order to track
+ * custom schema handler settings (like ordering, which is enabled, etc.)
+ */
+USTRUCT(BlueprintType)
+struct FSchemaHandlerEntry
+{
+	GENERATED_USTRUCT_BODY()
+
+public:
+	UPROPERTY(EditAnywhere, BlueprintReadOnly, Category = "USD Translator")
+	FString HandlerName;
+
+	UPROPERTY(EditAnywhere, BlueprintReadWrite, Category = "USD Translator")
+	bool bEnabled = true;
+
+	UPROPERTY(EditAnywhere, BlueprintReadOnly, Category = "USD Translator")
+	FString SchemaName;
+
+	/** For material schema handlers, we also hoist the render contexts the handler supports here, to be visible on the UI */
+	UPROPERTY(EditAnywhere, BlueprintReadOnly, Category = "USD Translator")
+	TArray<FString> DefaultRenderContexts;
+};
diff --git a/Engine/Plugins/Interchange/Extensions/OpenUSD/Source/Import/Public/SchemaHandlers/SchemaHandlerRegistry.h b/Engine/Plugins/Interchange/Extensions/OpenUSD/Source/Import/Public/SchemaHandlers/SchemaHandlerRegistry.h
new file mode 100644
--- /dev/null
+++ b/Engine/Plugins/Interchange/Extensions/OpenUSD/Source/Import/Public/SchemaHandlers/SchemaHandlerRegistry.h
@@ -0,0 +1,95 @@
+// Copyright Epic Games, Inc. All Rights Reserved.
+
+#pragma once
+
+#if USE_USD_SDK
+
+#include "SchemaHandlers/SchemaHandlerEntry.h"
+#include "SchemaHandlers/SchemaHandler.h"
+#include "SchemaHandlers/MaterialSchemaHandler.h"
+
+#include "Containers/AnsiString.h"
+
+#define UE_API INTERCHANGEOPENUSDIMPORT_API
+
+namespace UE::Interchange::USD
+{
+	/** Main lambda invoked when spawning an FSchemaHandler */
+	using FSchemaHandlerGenerator = TFunction<TSharedRef<FSchemaHandler>()>;
+
+	/**
+	 * Static storage of all registered FSchemaHandlers.
+	 */
+	class FSchemaHandlerRegistry
+	{
+	public:
+		/**
+		 * Variable that stores the actual registered schema handlers.
+		 * This is intentionally exposed so as to allow free insertion/reorder/removal of registered handlers.
+		 */
+		static UE_API TArray<FSchemaHandlerEntry> RegisteredHandlerEntries;
+
+		/**
+		 * Separate mapping of names to generators. This is useful because it lets us freely copy and serialize
+		 * our handler entries to reorder / configure them.
+		 *
+		 * Registering your handler via FSchemaHandlerRegistry::Register() will also register it here.
+		 *
+		 * Whenever a handler is executed without having a corresponding generator function registered for its name,
+		 * a warning will be emitted and the handler will be skipped.
+		 */
+		static UE_API TMap<FString, FSchemaHandlerGenerator> RegisteredHandlerNamesToGenerators;
+
+	public:
+		/**
+		 * Registers a schema handler of a particular type, so that it is invoked when translating USD files.
+		 * Usage:
+		 *     FString HandlerName = FSchemaHandlerRegistry::Register<FImageableSchemaHandler>();
+		 */
+		template<typename SchemaHandlerType>
+		static FString Register()
+		{
+			SchemaHandlerType TempHandler{};
+			const FString& HandlerName = TempHandler.GetHandlerName();
+
+			FSchemaHandlerEntry& NewEntry = RegisteredHandlerEntries.Emplace_GetRef();
+			NewEntry.HandlerName = HandlerName;
+			NewEntry.bEnabled = TempHandler.IsEnabled();
+			NewEntry.SchemaName = TempHandler.GetTargetSchemaName();
+			if constexpr (std::is_base_of_v<FMaterialSchemaHandler, SchemaHandlerType>)
+			{
+				NewEntry.DefaultRenderContexts = static_cast<FMaterialSchemaHandler*>(&TempHandler)->GetDefaultRenderContexts();
+			}
+
+			FSchemaHandlerGenerator GeneratorFunc = []()
+			{
+				TSharedRef<FSchemaHandler> Result = MakeShared<SchemaHandlerType>();
+				return Result;
+			};
+			RegisteredHandlerNamesToGenerators.Add(HandlerName, GeneratorFunc);
+
+			return HandlerName;
+		}
+
+		/**
+		 * Unregisters a schema handler, so that it will no longer be invoked when translating USD files.
+		 * Returns true if it unregistered anything.
+		 */
+		static UE_API bool Unregister(const FString& HandlerName);
+
+		/**
+		 * Spawns an instance of each entry in RegisteredHandlerEntries, in order.
+		 *
+		 * By default this will use the registered handler entries, but a different array of entries can be provided.
+		 *
+		 * Before translation, the UInterchangeUSDTranslator calls this function and retains those same spawned handlers
+		 * throughout the entire import process, so they are allowed to keep internal state between translation and
+		 * payload retrieval
+		 */
+		static UE_API TArray<TSharedRef<FSchemaHandler>> GenerateHandlers(const TArray<FSchemaHandlerEntry>* Entries = &RegisteredHandlerEntries);
+	};
+}	 // namespace UE::Interchange::USD
+
+#undef UE_API
+
+#endif	  // USE_USD_SDK
diff --git a/Engine/Plugins/Interchange/Extensions/OpenUSD/Source/Import/Public/SchemaHandlers/UniversalMaterialSchemaHandler.h b/Engine/Plugins/Interchange/Extensions/OpenUSD/Source/Import/Public/SchemaHandlers/UniversalMaterialSchemaHandler.h
--- a/Engine/Plugins/Interchange/Extensions/OpenUSD/Source/Import/Public/SchemaHandlers/UniversalMaterialSchemaHandler.h
+++ b/Engine/Plugins/Interchange/Extensions/OpenUSD/Source/Import/Public/SchemaHandlers/UniversalMaterialSchemaHandler.h
@@ -4,18 +4,18 @@
 
 #if USE_USD_SDK
 
-#include "SchemaHandler.h"
+#include "SchemaHandlers/MaterialSchemaHandler.h"
 
 #define UE_API INTERCHANGEOPENUSDIMPORT_API
 
 namespace UE::Interchange::USD
 {
-	class FUniversalMaterialSchemaHandler : public FSchemaHandler
+	class FUniversalMaterialSchemaHandler : public FMaterialSchemaHandler
 	{
 	public:
 		UE_API virtual const FString& GetHandlerName() const override;
 
-		UE_API virtual const FString& GetTargetSchemaName() const override;
+		UE_API const TArray<FString>& GetDefaultRenderContexts() const override;
 
 		UE_API virtual TOptional<bool> CanBeCollapsed(const UE::FUsdPrim& Prim, UInterchangeUsdContext& UsdContext) const override;
 
diff --git a/Engine/Plugins/Interchange/Extensions/OpenUSD/Source/Import/Public/SchemaHandlers/UnrealMaterialSchemaHandler.h b/Engine/Plugins/Interchange/Extensions/OpenUSD/Source/Import/Public/SchemaHandlers/UnrealMaterialSchemaHandler.h
--- a/Engine/Plugins/Interchange/Extensions/OpenUSD/Source/Import/Public/SchemaHandlers/UnrealMaterialSchemaHandler.h
+++ b/Engine/Plugins/Interchange/Extensions/OpenUSD/Source/Import/Public/SchemaHandlers/UnrealMaterialSchemaHandler.h
@@ -4,20 +4,18 @@
 
 #if USE_USD_SDK
 
-#include "SchemaHandler.h"
+#include "SchemaHandlers/MaterialSchemaHandler.h"
 
 #define UE_API INTERCHANGEOPENUSDIMPORT_API
 
 namespace UE::Interchange::USD
 {
-	class FUnrealMaterialSchemaHandler : public FSchemaHandler
+	class FUnrealMaterialSchemaHandler : public FMaterialSchemaHandler
 	{
 	public:
 		UE_API virtual const FString& GetHandlerName() const override;
 
-		UE_API virtual const FString& GetTargetSchemaName() const override;
-
-		UE_API virtual bool CanHandlePrim(const UE::FUsdPrim& Prim, const UInterchangeUsdContext& UsdContext) const override;
+		UE_API virtual const TArray<FString>& GetDefaultRenderContexts() const override;
 
 		UE_API virtual bool OnTranslate(
 			const UE::FUsdPrim& Prim,
diff --git a/Engine/Plugins/Runtime/USDCore/Source/USDUtilities/Private/USDGeomMeshConversion.cpp b/Engine/Plugins/Runtime/USDCore/Source/USDUtilities/Private/USDGeomMeshConversion.cpp
--- a/Engine/Plugins/Runtime/USDCore/Source/USDUtilities/Private/USDGeomMeshConversion.cpp
+++ b/Engine/Plugins/Runtime/USDCore/Source/USDUtilities/Private/USDGeomMeshConversion.cpp
@@ -4221,7 +4221,7 @@ UsdUtils::FUsdPrimMaterialAssignmentInfo UsdUtils::GetPrimMaterialAssignments(
 		return {};
 	};
 
-	auto FetchMaterialByComputingBoundMaterial = [&RenderContext, &MaterialPurpose](const pxr::UsdPrim& UsdPrim) -> TOptional<FString>
+	auto FetchMaterialByComputingBoundMaterial = [&MaterialPurpose](const pxr::UsdPrim& UsdPrim) -> TOptional<FString>
 	{
 		pxr::UsdShadeMaterialBindingAPI BindingAPI(UsdPrim);
 		pxr::UsdShadeMaterial ShadeMaterial = BindingAPI.ComputeBoundMaterial(MaterialPurpose);
@@ -4230,17 +4230,6 @@ UsdUtils::FUsdPrimMaterialAssignmentInfo UsdUtils::GetPrimMaterialAssignments(
 			return {};
 		}
 
-		if (RenderContext.GetString().c_str() != UnrealIdentifiers::MaterialXRenderContext)
-		{
-			// Ignore this material if UsdToUnreal::ConvertMaterial would as well 
-			// in a MaterialX context we can have any type of material though (surface/displacement/volume)
-			pxr::UsdShadeShader SurfaceShader = ShadeMaterial.ComputeSurfaceSource(RenderContext);
-			if (!SurfaceShader)
-			{
-				return {};
-			}
-		}
-
 		pxr::UsdPrim ShadeMaterialPrim = ShadeMaterial.GetPrim();
 		if (ShadeMaterialPrim)
 		{
diff --git a/Engine/Plugins/Runtime/USDCore/Source/USDUtilities/Private/USDShadeConversion.cpp b/Engine/Plugins/Runtime/USDCore/Source/USDUtilities/Private/USDShadeConversion.cpp
--- a/Engine/Plugins/Runtime/USDCore/Source/USDUtilities/Private/USDShadeConversion.cpp
+++ b/Engine/Plugins/Runtime/USDCore/Source/USDUtilities/Private/USDShadeConversion.cpp
@@ -3492,6 +3492,33 @@ bool UsdUtils::HasSurfaceOutput(const pxr::UsdPrim& MaterialPrim, const FName& R
 	return static_cast<bool>(MaterialOutput);
 }
 
+bool UsdUtils::HasSurfaceOutput(const pxr::UsdPrim& MaterialPrim, const TArray<FString>& RenderContexts)
+{
+	if (!MaterialPrim)
+	{
+		return false;
+	}
+
+	FScopedUsdAllocs UsdAllocs;
+
+	pxr::UsdShadeMaterial ShadeMaterial{MaterialPrim};
+	if (!ShadeMaterial)
+	{
+		return false;
+	}
+
+	for (const FString& RenderContext : RenderContexts)
+	{
+		pxr::UsdShadeOutput MaterialOutput = ShadeMaterial.GetSurfaceOutput(UnrealToUsd::ConvertToken(*RenderContext).Get());
+		if (MaterialOutput)
+		{
+			return true;
+		}
+	}
+
+	return false;
+}
+
 bool UsdUtils::HasDisplacementOutput(const pxr::UsdPrim& MaterialPrim, const FName& RenderContext)
 {
 	if (!MaterialPrim)
@@ -3511,6 +3538,33 @@ bool UsdUtils::HasDisplacementOutput(const pxr::UsdPrim& MaterialPrim, const FNa
 	return static_cast<bool>(MaterialOutput);
 }
 
+bool UsdUtils::HasDisplacementOutput(const pxr::UsdPrim& MaterialPrim, const TArray<FString>& RenderContexts)
+{
+	if (!MaterialPrim)
+	{
+		return false;
+	}
+
+	FScopedUsdAllocs UsdAllocs;
+
+	pxr::UsdShadeMaterial ShadeMaterial{MaterialPrim};
+	if (!ShadeMaterial)
+	{
+		return false;
+	}
+
+	for (const FString& RenderContext : RenderContexts)
+	{
+		pxr::UsdShadeOutput MaterialOutput = ShadeMaterial.GetDisplacementOutput(UnrealToUsd::ConvertToken(*RenderContext).Get());
+		if (MaterialOutput)
+		{
+			return true;
+		}
+	}
+
+	return false;
+}
+
 bool UsdUtils::HasVolumeOutput(const pxr::UsdPrim& MaterialPrim, const FName& RenderContext)
 {
 	if (!MaterialPrim)
@@ -3530,6 +3584,33 @@ bool UsdUtils::HasVolumeOutput(const pxr::UsdPrim& MaterialPrim, const FName& Re
 	return static_cast<bool>(MaterialOutput);
 }
 
+bool UsdUtils::HasVolumeOutput(const pxr::UsdPrim& MaterialPrim, const TArray<FString>& RenderContexts)
+{
+	if (!MaterialPrim)
+	{
+		return false;
+	}
+
+	FScopedUsdAllocs UsdAllocs;
+
+	pxr::UsdShadeMaterial ShadeMaterial{MaterialPrim};
+	if (!ShadeMaterial)
+	{
+		return false;
+	}
+
+	for (const FString& RenderContext : RenderContexts)
+	{
+		pxr::UsdShadeOutput MaterialOutput = ShadeMaterial.GetVolumeOutput(UnrealToUsd::ConvertToken(*RenderContext).Get());
+		if (MaterialOutput)
+		{
+			return true;
+		}
+	}
+
+	return false;
+}
+
 bool UsdUtils::IsMaterialTranslucent(const pxr::UsdShadeMaterial& UsdShadeMaterial)
 {
 	FScopedUsdAllocs UsdAllocs;
diff --git a/Engine/Plugins/Runtime/USDCore/Source/USDUtilities/Public/USDShadeConversion.h b/Engine/Plugins/Runtime/USDCore/Source/USDUtilities/Public/USDShadeConversion.h
--- a/Engine/Plugins/Runtime/USDCore/Source/USDUtilities/Public/USDShadeConversion.h
+++ b/Engine/Plugins/Runtime/USDCore/Source/USDUtilities/Public/USDShadeConversion.h
@@ -274,12 +274,15 @@ namespace UsdUtils
 
 	// Returns whether MaterialPrim is an actual Material, and has a surface output authored for the provided render context
 	USDUTILITIES_API bool HasSurfaceOutput(const pxr::UsdPrim& MaterialPrim, const FName& RenderContext);
-
+	USDUTILITIES_API bool HasSurfaceOutput(const pxr::UsdPrim& MaterialPrim, const TArray<FString>& RenderContexts);
+	
 	// Returns whether MaterialPrim is an actual Material, and has a displacement output authored for the provided render context
 	USDUTILITIES_API bool HasDisplacementOutput(const pxr::UsdPrim& MaterialPrim, const FName& RenderContext);
-
+	USDUTILITIES_API bool HasDisplacementOutput(const pxr::UsdPrim& MaterialPrim, const TArray<FString>& RenderContexts);
+	
 	// Returns whether MaterialPrim is an actual Material, and has a volume output authored for the provided render context
 	USDUTILITIES_API bool HasVolumeOutput(const pxr::UsdPrim& MaterialPrim, const FName& RenderContext);
+	USDUTILITIES_API bool HasVolumeOutput(const pxr::UsdPrim& MaterialPrim, const TArray<FString>& RenderContexts);
 
 	/**
 	 * Returns whether the material needs to be rendered with the Translucent rendering mode.
