USD Interchange: Move geometry purpose handling back to the translator

#Synced-CL 49553638

diff --git a/Engine/Plugins/Interchange/Extensions/OpenUSD/Source/Import/Private/InterchangeUSDPipeline.cpp b/Engine/Plugins/Interchange/Extensions/OpenUSD/Source/Import/Private/InterchangeUSDPipeline.cpp
--- a/Engine/Plugins/Interchange/Extensions/OpenUSD/Source/Import/Private/InterchangeUSDPipeline.cpp
+++ b/Engine/Plugins/Interchange/Extensions/OpenUSD/Source/Import/Private/InterchangeUSDPipeline.cpp
@@ -66,167 +66,6 @@ namespace UE::InterchangeUsdPipeline::Private
 		return Result;
 	}
 
-	// Disables the factory nodes on NodeContainer if they target translated nodes that had disallowed purposes.
-	//
-	// Our payload retrieval functions are also filtering and not producing mesh data for disabled purpose meshes,
-	// so theoretically this shouldn't be necessary. However, if a factory node exists and is enabled, Interchange
-	// will always create a StaticMesh for it and fill it with the provided MeshDescription. If the entire mesh description
-	// is empty, the StaticMesh will be invalid and emit a bunch of errors.
-	//
-	// So essentially here we disable factory nodes whose meshes have been *fully* disabled due to purpose, and when retrieving
-	// payloads we also double-check the purpose so that we can ignore individual meshes in the context of collapsing and combining
-	// larger meshes
-	void DisableNodesBasedOnGeometryPurpose(UInterchangeBaseNodeContainer& NodeContainer, const EUsdPurpose AllowedPurposes)
-	{
-		// Collect all the disabled mesh, meshbundle and scene nodes
-		TSet<FString> DisabledTranslatedNodes;
-		NodeContainer.IterateNodes(
-			[&DisabledTranslatedNodes, AllowedPurposes, &NodeContainer](const FString& NodeUid, UInterchangeBaseNode* BaseNode)
-			{
-				FString AssetNodeUid;
-				if (UInterchangeSceneNode* SceneNode = Cast<UInterchangeSceneNode>(BaseNode))
-				{
-					SceneNode->GetCustomAssetInstanceUid(AssetNodeUid);
-				}
-				else if (UInterchangeInstancedStaticMeshComponentNode* SceneComponentNode = Cast<UInterchangeInstancedStaticMeshComponentNode>(BaseNode))
-				{
-					SceneComponentNode->GetCustomInstancedAssetUid(AssetNodeUid);
-				}
-				else
-				{
-					return;
-				}
-
-				const UInterchangeBaseNode* AssetNode = NodeContainer.GetNode(AssetNodeUid);
-				if (!AssetNode)
-				{
-					return;
-				}
-
-				// Single mesh scene node: Mark both the scene node and the asset node as disabled
-				if (AssetNode->IsA<UInterchangeMeshNode>())
-				{
-					EUsdPurpose AuthoredPurpose = EUsdPurpose::Default;
-					bool bHasPurpose = BaseNode->GetInt32Attribute(UE::Interchange::USD::GeometryPurposeIdentifier, (int32&)AuthoredPurpose);
-					if (bHasPurpose && !EnumHasAllFlags(AllowedPurposes, AuthoredPurpose))
-					{
-						DisabledTranslatedNodes.Add(NodeUid);
-						DisabledTranslatedNodes.Add(AssetNodeUid);
-					}
-				}
-				else if (const UInterchangeMeshBundleNode* MeshBundle = Cast<UInterchangeMeshBundleNode>(AssetNode))
-				{
-					// Mesh bundle: Check if *all* of our meshes are disabled: In that case we need to disable the bundle / scene node
-					// too. If only some of the nodes are disabled then the Gprim handler payload retrieval functions will already take
-					// care of producing empty MeshDescriptions for those individual sections, but we also need this here because we
-					// can't provide completely empty MeshDescription data to the static mesh factory
-					TArray<FString> MeshNodeUids;
-					MeshBundle->GetMeshNodeUids(MeshNodeUids);
-
-					bool bAllMeshesDisabled = true;
-					for (const FString& MeshNodeUid : MeshNodeUids)
-					{
-						if (const UInterchangeMeshNode* BundleMeshNode = Cast<UInterchangeMeshNode>(NodeContainer.GetNode(MeshNodeUid)))
-						{
-							EUsdPurpose AuthoredPurpose = EUsdPurpose::Default;
-							bool bHasPurpose = BundleMeshNode->GetInt32Attribute(UE::Interchange::USD::GeometryPurposeIdentifier, (int32&)AuthoredPurpose);
-							if (!bHasPurpose || EnumHasAllFlags(AllowedPurposes, AuthoredPurpose))
-							{
-								bAllMeshesDisabled = false;
-								break;
-							}
-						}
-					}
-
-					if (bAllMeshesDisabled)
-					{
-						DisabledTranslatedNodes.Add(NodeUid);
-						DisabledTranslatedNodes.Add(AssetNodeUid);
-					}
-				}
-			}
-		);
-
-		// Disable any factory node that targets/references one of our disabled scene nodes.
-		//
-		// We do this in a separate pass because we can say absolutely nothing about our factory nodes, as any number
-		// of pipelines may have done arbirary transformations before we got to run. We have to hope they kept the
-		// target node attribute updated at least, and use that.
-		NodeContainer.IterateNodesOfType<UInterchangeFactoryBaseNode>(
-			[&DisabledTranslatedNodes](const FString& NodeUid, UInterchangeFactoryBaseNode* FactoryNode)
-			{
-				// Handle standard factory nodes
-				{
-					TArray<FString> TargetTranslatedNodeUids;
-					FactoryNode->GetTargetNodeUids(TargetTranslatedNodeUids);
-
-					bool bRemovedNode = false;
-					for (const FString& TargetNode : TargetTranslatedNodeUids)
-					{
-						if (DisabledTranslatedNodes.Contains(TargetNode))
-						{
-							bRemovedNode = true;
-							FactoryNode->RemoveTargetNodeUid(TargetNode);
-						}
-					}
-
-					// We don't want to leave a factory node enabled without a target, as that may lead
-					// to some errors on the factories
-					if (bRemovedNode && FactoryNode->GetTargetNodeCount() == 0)
-					{
-						const bool bEnabled = false;
-						FactoryNode->SetEnabled(bEnabled);
-					}
-				}
-
-				// These nodes don't use the common "target" mechanism and store their mesh nodes separately...
-				// TODO: Will have to expand these to also handle geometry caches?
-				if (UInterchangeStaticMeshLodDataNode* LodDataNode = Cast<UInterchangeStaticMeshLodDataNode>(FactoryNode))
-				{
-					TArray<FString> TranslatedMeshNodeUids;
-					LodDataNode->GetMeshUids(TranslatedMeshNodeUids);
-
-					bool bRemovedNode = false;
-					for (const FString& MeshNodeUid : TranslatedMeshNodeUids)
-					{
-						if (DisabledTranslatedNodes.Contains(MeshNodeUid))
-						{
-							bRemovedNode = true;
-							LodDataNode->RemoveMeshUid(MeshNodeUid);
-						}
-					}
-
-					if (bRemovedNode && LodDataNode->GetMeshUidsCount() == 0)
-					{
-						const bool bEnabled = false;
-						LodDataNode->SetEnabled(bEnabled);
-					}
-				}
-				else if (UInterchangeSkeletalMeshLodDataNode* SkeletalNode = Cast<UInterchangeSkeletalMeshLodDataNode>(FactoryNode))
-				{
-					TArray<FString> TranslatedMeshNodeUids;
-					SkeletalNode->GetMeshUids(TranslatedMeshNodeUids);
-
-					bool bRemovedNode = false;
-					for (const FString& MeshNodeUid : TranslatedMeshNodeUids)
-					{
-						if (DisabledTranslatedNodes.Contains(MeshNodeUid))
-						{
-							bRemovedNode = true;
-							SkeletalNode->RemoveMeshUid(MeshNodeUid);
-						}
-					}
-
-					if (bRemovedNode && SkeletalNode->GetMeshUidsCount() == 0)
-					{
-						const bool bEnabled = false;
-						SkeletalNode->SetEnabled(bEnabled);
-					}
-				}
-			}
-		);
-	}
-
 	// Moves our primvar-compatible custom attributes from our flagged material instance translated nodes to the
 	// corresponding factory nodes.
 	//
@@ -1010,8 +849,7 @@ namespace UE::InterchangeUsdPipeline::Private
 		UInterchangeBaseNodeContainer* NodeContainer,
 		bool bGenerateCompatibleMaterials,
 		EInterchangeUsdPrimvar ImportPrimvar,
-		int32 SubdivisionLevel,
-		int32 GeometryPurpose
+		int32 SubdivisionLevel
 	)
 	{
 		using namespace UE::Interchange;
@@ -1026,7 +864,6 @@ namespace UE::InterchangeUsdPipeline::Private
 			 &MeshNodeToMeshFactoryNode,
 			 ImportPrimvar,
 			 SubdivisionLevel,
-			 GeometryPurpose,
 			 bGenerateCompatibleMaterials,
 			 &MeshNodeToMeshActorFactoryNode](const FString& NodeUid, UInterchangeBaseNode* BaseNode)
 			{
@@ -1097,11 +934,6 @@ namespace UE::InterchangeUsdPipeline::Private
 						}
 					}
 				}
-				// Check the mesh factory node directly as we don't care whether it came from a Mesh node or a MeshBundle node
-				else if (UInterchangeMeshFactoryNode* MeshFactoryNode = Cast<UInterchangeMeshFactoryNode>(BaseNode))
-				{
-					MeshFactoryNode->AddPayloadKeyInt32Attribute(UE::Interchange::USD::GeometryPurposeIdentifier, GeometryPurpose);
-				}
 			}
 		);
 
@@ -1446,8 +1278,7 @@ namespace UE::InterchangeUsdPipeline::Private
 }	 // namespace UE::InterchangeUsdPipeline::Private
 
 UInterchangeUsdPipeline::UInterchangeUsdPipeline()
-	: GeometryPurpose((int32)(EUsdPurpose::Default | EUsdPurpose::Proxy | EUsdPurpose::Render | EUsdPurpose::Guide))
-	, SubdivisionLevel(0)
+	: SubdivisionLevel(0)
 	, bImportPseudoRoot(true)
 	, bGeneratePrimvarCompatibleMaterials(true)
 {
@@ -1581,8 +1412,6 @@ void UInterchangeUsdPipeline::ExecutePipeline(
 		return;
 	}
 
-	DisableNodesBasedOnGeometryPurpose(*NodeContainer, (EUsdPurpose)GeometryPurpose);
-
 	MoveCustomAttributesToMaterialFactoryNodes(*NodeContainer);
 
 	ProcessVolumeNodes(*NodeContainer);
@@ -1592,7 +1421,7 @@ void UInterchangeUsdPipeline::ExecutePipeline(
 		RemovePseudoRootFactoryNode(*NodeContainer);
 	}
 
-	ProcessMeshNodes(NodeContainer, bGeneratePrimvarCompatibleMaterials, ImportPrimvars, SubdivisionLevel, GeometryPurpose);
+	ProcessMeshNodes(NodeContainer, bGeneratePrimvarCompatibleMaterials, ImportPrimvars, SubdivisionLevel);
 #endif	  // USE_USD_SDK
 }
 
diff --git a/Engine/Plugins/Interchange/Extensions/OpenUSD/Source/Import/Private/InterchangeUsdContext.cpp b/Engine/Plugins/Interchange/Extensions/OpenUSD/Source/Import/Private/InterchangeUsdContext.cpp
--- a/Engine/Plugins/Interchange/Extensions/OpenUSD/Source/Import/Private/InterchangeUsdContext.cpp
+++ b/Engine/Plugins/Interchange/Extensions/OpenUSD/Source/Import/Private/InterchangeUsdContext.cpp
@@ -268,9 +268,6 @@ FInterchangeUsdInfoCache* UInterchangeUsdContext::GetInterchangeInfoCache() cons
 void UInterchangeUsdContext::SetupMeshConversionOptions()
 {
 #if USE_USD_SDK
-	// We filter for this on the pipeline now
-	CachedMeshConversionOptions.PurposesToLoad = EUsdPurpose::Default | EUsdPurpose::Proxy | EUsdPurpose::Render | EUsdPurpose::Guide;
-
 	if (UInterchangeUsdTranslatorSettings* Settings = GetTranslatorSettings())
 	{
 		// TODO: Change FUsdMeshConversionOptions to not hold USD types directly, so we don't have to the conversion below everywhere.
@@ -278,6 +275,7 @@ void UInterchangeUsdContext::SetupMeshConversionOptions()
 		// really do anything anyway as the module doesn't use IMPLEMENT_MODULE_USD!
 		// Luckily we can get around this here because pxr::TfToken doesn't allocate on its own: At most USD makes a copy of the string, which it
 		// should allocate/deallocate on its own allocator.
+		CachedMeshConversionOptions.PurposesToLoad = static_cast<EUsdPurpose>(Settings->GeometryPurpose);
 		CachedMeshConversionOptions.MaterialPurpose = Settings->MaterialPurpose.IsNone()
 														  ? pxr::UsdShadeTokens->allPurpose
 														  : pxr::TfToken{TCHAR_TO_UTF8(*Settings->MaterialPurpose.ToString())};
diff --git a/Engine/Plugins/Interchange/Extensions/OpenUSD/Source/Import/Private/SchemaHandlers/GeometryCacheSchemaHandler.cpp b/Engine/Plugins/Interchange/Extensions/OpenUSD/Source/Import/Private/SchemaHandlers/GeometryCacheSchemaHandler.cpp
--- a/Engine/Plugins/Interchange/Extensions/OpenUSD/Source/Import/Private/SchemaHandlers/GeometryCacheSchemaHandler.cpp
+++ b/Engine/Plugins/Interchange/Extensions/OpenUSD/Source/Import/Private/SchemaHandlers/GeometryCacheSchemaHandler.cpp
@@ -5,6 +5,7 @@
 #include "SchemaHandlers/GeometryCacheSchemaHandler.h"
 
 #include "InterchangeUsdContext.h"
+#include "InterchangeUsdTranslator.h"
 #include "InterchangeUsdTraversalInfo.h"
 #include "SchemaHandlers/SchemaHandlerUtils.h"
 #include "Usd/InterchangeUsdDefinitions.h"
@@ -163,6 +164,15 @@ namespace UE::Interchange::USD
 
 		using namespace UE::GeometryCacheSchemaHandler::Private;
 
+		// Don't process mesh prims with disabled purposes
+		if (UInterchangeUsdTranslatorSettings* Settings = UsdContext.GetTranslatorSettings())
+		{
+			if (!EnumHasAllFlags(static_cast<EUsdPurpose>(Settings->GeometryPurpose), IUsdPrim::GetPurpose(Prim)))
+			{
+				return false;
+			}
+		}
+
 		// Go through the prototype (if any) so that we try reusing the same mesh asset node for different USD instances
 		UE::FUsdPrim PrototypePrim = UsdUtils::GetPrototypePrim(Prim);
 		const FString PrototypePrimPath = PrototypePrim.GetPrimPath().GetString();
@@ -207,7 +217,6 @@ namespace UE::Interchange::USD
 
 		PayloadAttributes.GetAttribute(UE::Interchange::FAttributeKey{USD::SubdivisionLevelAttributeKey}, OptionsCopy.SubdivisionLevel);
 		PayloadAttributes.GetAttribute(UE::Interchange::FAttributeKey{USD::Primvar::Import}, reinterpret_cast<int32&>(OptionsCopy.ImportPrimvars));
-		PayloadAttributes.GetAttribute(UE::Interchange::FAttributeKey{USD::GeometryPurposeIdentifier}, reinterpret_cast<int32&>(OptionsCopy.PurposesToLoad));
 
 		if (EInterchangeUsdPrimvar(OptionsCopy.ImportPrimvars) != EInterchangeUsdPrimvar::Standard)
 		{
diff --git a/Engine/Plugins/Interchange/Extensions/OpenUSD/Source/Import/Private/SchemaHandlers/GprimSchemaHandler.cpp b/Engine/Plugins/Interchange/Extensions/OpenUSD/Source/Import/Private/SchemaHandlers/GprimSchemaHandler.cpp
--- a/Engine/Plugins/Interchange/Extensions/OpenUSD/Source/Import/Private/SchemaHandlers/GprimSchemaHandler.cpp
+++ b/Engine/Plugins/Interchange/Extensions/OpenUSD/Source/Import/Private/SchemaHandlers/GprimSchemaHandler.cpp
@@ -5,6 +5,7 @@
 #include "SchemaHandlers/GprimSchemaHandler.h"
 
 #include "InterchangeUsdContext.h"
+#include "InterchangeUsdTranslator.h"
 #include "InterchangeUsdTraversalInfo.h"
 #include "SchemaHandlers/SchemaHandlerUtils.h"
 #include "Usd/InterchangeUsdDefinitions.h"
@@ -274,14 +275,6 @@ namespace UE::GeomGprimSchemaHandler::Private
 			}
 		}
 
-		// Store our purpose on the mesh node as well (the imageable handler will also store it on the scene node)
-		// as this is needed to be able to disable parts of collapsed meshes on the USDPipeline
-		const EUsdPurpose PrimPurpose = IUsdPrim::GetPurpose(MeshPrim);
-		if (PrimPurpose != EUsdPurpose::Default)
-		{
-			MeshNode->AddInt32Attribute(UE::Interchange::USD::GeometryPurposeIdentifier, (int32)PrimPurpose);
-		}
-
 		// Material assignments
 		{
 			double TimeCode = UsdUtils::GetDefaultTimeCode();
@@ -675,6 +668,15 @@ namespace UE::Interchange::USD
 		const bool bIsPrimitiveShape = !Prim.IsA(TEXT("Mesh")); // See CanHandlePrim()
 		const bool bIsAnimated = UsdUtils::IsAnimatedMesh(Prim);
 
+		// Don't process mesh prims with disabled purposes
+		if (UInterchangeUsdTranslatorSettings* Settings = UsdContext.GetTranslatorSettings())
+		{
+			if (!EnumHasAllFlags(static_cast<EUsdPurpose>(Settings->GeometryPurpose), IUsdPrim::GetPurpose(Prim)))
+			{
+				return false;
+			}
+		}
+
 		UE::FUsdPrim PrototypePrim;
 		if (bIsSkinned || TraversalInfo.bInsideLODVariant)
 		{
@@ -794,7 +796,6 @@ namespace UE::Interchange::USD
 
 		PayloadAttributes.GetAttribute(UE::Interchange::FAttributeKey{USD::SubdivisionLevelAttributeKey}, OptionsCopy.SubdivisionLevel);
 		PayloadAttributes.GetAttribute(UE::Interchange::FAttributeKey{USD::Primvar::Import}, reinterpret_cast<int32&>(OptionsCopy.ImportPrimvars));
-		PayloadAttributes.GetAttribute(UE::Interchange::FAttributeKey{USD::GeometryPurposeIdentifier}, reinterpret_cast<int32&>(OptionsCopy.PurposesToLoad));
 
 		if (EInterchangeUsdPrimvar(OptionsCopy.ImportPrimvars) != EInterchangeUsdPrimvar::Standard)
 		{
diff --git a/Engine/Plugins/Interchange/Extensions/OpenUSD/Source/Import/Private/SchemaHandlers/ImageableSchemaHandler.cpp b/Engine/Plugins/Interchange/Extensions/OpenUSD/Source/Import/Private/SchemaHandlers/ImageableSchemaHandler.cpp
--- a/Engine/Plugins/Interchange/Extensions/OpenUSD/Source/Import/Private/SchemaHandlers/ImageableSchemaHandler.cpp
+++ b/Engine/Plugins/Interchange/Extensions/OpenUSD/Source/Import/Private/SchemaHandlers/ImageableSchemaHandler.cpp
@@ -484,13 +484,6 @@ namespace UE::Interchange::USD
 		};
 		AddNodesForAnimatedAttributes(Prim, AttributeMapping, AccumulatedInfo, UsdContext);
 
-		// Store our purpose here to be filtered on the pipeline
-		const EUsdPurpose PrimPurpose = IUsdPrim::GetPurpose(Prim);
-		if (SceneNodeBase && PrimPurpose != EUsdPurpose::Default)
-		{
-			SceneNodeBase->AddInt32Attribute(UE::Interchange::USD::GeometryPurposeIdentifier, (int32)PrimPurpose);
-		}
-
 		// LOD handling
 		if (UInterchangeSceneNode* SceneNode = Cast<UInterchangeSceneNode>(SceneNodeBase))
 		{
diff --git a/Engine/Plugins/Interchange/Extensions/OpenUSD/Source/Import/Private/SchemaHandlers/PointInstancerSchemaHandler.cpp b/Engine/Plugins/Interchange/Extensions/OpenUSD/Source/Import/Private/SchemaHandlers/PointInstancerSchemaHandler.cpp
--- a/Engine/Plugins/Interchange/Extensions/OpenUSD/Source/Import/Private/SchemaHandlers/PointInstancerSchemaHandler.cpp
+++ b/Engine/Plugins/Interchange/Extensions/OpenUSD/Source/Import/Private/SchemaHandlers/PointInstancerSchemaHandler.cpp
@@ -102,13 +102,6 @@ namespace UE::GeomPointInstancerSchemaHandler::Private
 			UidToComponentNode.Add(UniqueId, ISMComponentNode);
 			AccumulatedInfo.PrimSceneNodes.Add(ISMComponentNode);
 
-			// Store our purpose here to be filtered on the pipeline
-			const EUsdPurpose PrimPurpose = IUsdPrim::GetPurpose(Prim);
-			if (ISMComponentNode && PrimPurpose != EUsdPurpose::Default)
-			{
-				ISMComponentNode->AddInt32Attribute(UE::Interchange::USD::GeometryPurposeIdentifier, (int32)PrimPurpose);
-			}
-
 			return ISMComponentNode;
 		}
 
diff --git a/Engine/Plugins/Interchange/Extensions/OpenUSD/Source/Import/Public/InterchangeUSDPipeline.h b/Engine/Plugins/Interchange/Extensions/OpenUSD/Source/Import/Public/InterchangeUSDPipeline.h
--- a/Engine/Plugins/Interchange/Extensions/OpenUSD/Source/Import/Public/InterchangeUSDPipeline.h
+++ b/Engine/Plugins/Interchange/Extensions/OpenUSD/Source/Import/Public/InterchangeUSDPipeline.h
@@ -27,7 +27,8 @@ public:
 	FString PipelineDisplayName;
 
 	/** Only import translated nodes from imageable prims (Xforms, Meshes, etc.) with these specific purposes from the USD file */
-	UPROPERTY(EditAnywhere, BlueprintReadWrite, Category = "USD", meta = (Bitmask, BitmaskEnum = "/Script/UnrealUSDWrapper.EUsdPurpose"))
+	UE_DEPRECATED(5.8, "This option has moved back to the USD translator settings")
+	UPROPERTY(meta = (Bitmask, BitmaskEnum = "/Script/UnrealUSDWrapper.EUsdPurpose"))
 	int32 GeometryPurpose;
 
 	/** Setting to tell what primvars have to be attached to the MeshDescription */
diff --git a/Engine/Plugins/Interchange/Extensions/OpenUSD/Source/Import/Public/InterchangeUsdTranslator.h b/Engine/Plugins/Interchange/Extensions/OpenUSD/Source/Import/Public/InterchangeUsdTranslator.h
--- a/Engine/Plugins/Interchange/Extensions/OpenUSD/Source/Import/Public/InterchangeUsdTranslator.h
+++ b/Engine/Plugins/Interchange/Extensions/OpenUSD/Source/Import/Public/InterchangeUsdTranslator.h
@@ -63,8 +63,7 @@ class UInterchangeUsdTranslatorSettings : public UInterchangeTranslatorSettings
 
 public:
 	/** Only import geometry prims with these specific purposes from the USD file */
-	UE_DEPRECATED(5.6, "This option is no longer a translator setting and has moved to the InterchangeUSDPipeline")
-	UPROPERTY()
+	UPROPERTY(EditAnywhere, Category = "USD Translator", meta = (Bitmask, BitmaskEnum = "/Script/UnrealUSDWrapper.EUsdPurpose"))
 	int32 GeometryPurpose;
 
 	/** Specifies which set of shaders to use when parsing USD materials, in addition to the universal render context. */
diff --git a/Engine/Plugins/Interchange/Runtime/Source/Common/Public/Usd/InterchangeUsdDefinitions.h b/Engine/Plugins/Interchange/Runtime/Source/Common/Public/Usd/InterchangeUsdDefinitions.h
--- a/Engine/Plugins/Interchange/Runtime/Source/Common/Public/Usd/InterchangeUsdDefinitions.h
+++ b/Engine/Plugins/Interchange/Runtime/Source/Common/Public/Usd/InterchangeUsdDefinitions.h
@@ -24,7 +24,10 @@ namespace UE::Interchange::USD
 	inline const FString USDContextTag = TEXT("USD");
 
 	// Name of a custom attribute added to translated nodes to contain their geometry purpose (proxy, render, guide, etc.)
+	PRAGMA_DISABLE_DEPRECATION_WARNINGS
+	UE_DEPRECATED(5.8, "No longer used, as geometry purpose is now handled on the translator and no longer needs to be delayed to the pipeline")
 	inline const FString GeometryPurposeIdentifier = TEXT("USD_Geometry_Purpose");
+	PRAGMA_ENABLE_DEPRECATION_WARNINGS
 
 	// Prefixes we use to stash some primvar mapping information as custom attributes on mesh / material nodes, so that
 	// the USD Pipeline can produce primvar-compatible materials
diff --git a/Engine/Plugins/Runtime/USDCore/Source/USDUtilities/Private/USDGeomMeshConversion.cpp b/Engine/Plugins/Runtime/USDCore/Source/USDUtilities/Private/USDGeomMeshConversion.cpp
--- a/Engine/Plugins/Runtime/USDCore/Source/USDUtilities/Private/USDGeomMeshConversion.cpp
+++ b/Engine/Plugins/Runtime/USDCore/Source/USDUtilities/Private/USDGeomMeshConversion.cpp
@@ -3447,12 +3447,6 @@ bool UsdToUnreal::ConvertGeomMesh(
 	FScopedUsdAllocs Allocs;
 
 	pxr::UsdPrim UsdPrim = UsdMesh.GetPrim();
-	const bool bPurposeIsVisible = EnumHasAllFlags(Options.PurposesToLoad, IUsdPrim::GetPurpose(UsdPrim));
-	if (!bPurposeIsVisible)
-	{
-		// This mesh shouldn't be visible, we've done what we should
-		return true;
-	}
 
 	UsdGeomMeshImpl::FUsdMeshData MeshData;
 
@@ -3779,13 +3773,6 @@ bool UsdToUnreal::ConvertGeomPrimitive(
 		return false;
 	}
 
-	const bool bPurposeIsVisible = EnumHasAllFlags(InOptions.PurposesToLoad, IUsdPrim::GetPurpose(InPrim));
-	if (!bPurposeIsVisible)
-	{
-		// This mesh shouldn't be visible, we've done what we should
-		return true;
-	}
-
 	FScopedUsdAllocs UsdAllocs;
 
 	UsdGeomMeshImpl::FUsdMeshData MeshData;
