Interchange USD: Fix texture asset names having their file extensions as suffixes.

#Synced-CL 49553638

diff --git a/Engine/Plugins/Interchange/Extensions/OpenUSD/Source/Import/Private/SchemaHandlers/LightSchemaHandler.cpp b/Engine/Plugins/Interchange/Extensions/OpenUSD/Source/Import/Private/SchemaHandlers/LightSchemaHandler.cpp
--- a/Engine/Plugins/Interchange/Extensions/OpenUSD/Source/Import/Private/SchemaHandlers/LightSchemaHandler.cpp
+++ b/Engine/Plugins/Interchange/Extensions/OpenUSD/Source/Import/Private/SchemaHandlers/LightSchemaHandler.cpp
@@ -46,7 +46,7 @@ namespace UE::LuxLightSchemaHandler::Private
 		UInterchangeTexture2DNode* Node = GetExistingNode<UInterchangeTexture2DNode>(NodeContainer, NodeUid);
 		if (!Node)
 		{
-			const FString NodeName{FPaths::GetCleanFilename(TextureFilePath)};
+			const FString NodeName{FPaths::GetBaseFilename(TextureFilePath)};
 			Node = NewObject<UInterchangeTexture2DNode>(&NodeContainer);
 			NodeContainer.SetupNode(Node, NodeUid, NodeName, EInterchangeNodeContainerType::TranslatedAsset);
 		}
diff --git a/Engine/Plugins/Interchange/Extensions/OpenUSD/Source/Import/Private/SchemaHandlers/UniversalMaterialSchemaHandler.cpp b/Engine/Plugins/Interchange/Extensions/OpenUSD/Source/Import/Private/SchemaHandlers/UniversalMaterialSchemaHandler.cpp
--- a/Engine/Plugins/Interchange/Extensions/OpenUSD/Source/Import/Private/SchemaHandlers/UniversalMaterialSchemaHandler.cpp
+++ b/Engine/Plugins/Interchange/Extensions/OpenUSD/Source/Import/Private/SchemaHandlers/UniversalMaterialSchemaHandler.cpp
@@ -41,7 +41,7 @@ namespace UE::UniversalMaterialSchemaHandler::Private
 	)
 	{
 		const FString PrimPath = UsdUtils::GetPrototypePrimPath(Prim).GetString();
-		const FString NodeName{FPaths::GetCleanFilename(Value.TextureFilePath)};
+		const FString NodeName{FPaths::GetBaseFilename(Value.TextureFilePath)};
 
 		UInterchangeTexture2DNode* Node = GetExistingNode<UInterchangeTexture2DNode>(NodeContainer, NodeUid);
 		if (!Node)
