USD: Rename EUsdPointInstancerCollapsing::DontCollapse to EUsdPointInstancerCollapsing::NoCollapsing

#Synced-CL 49553638

diff --git a/Engine/Plugins/Interchange/Extensions/OpenUSD/Source/Import/Private/InterchangeUsdTranslator.cpp b/Engine/Plugins/Interchange/Extensions/OpenUSD/Source/Import/Private/InterchangeUsdTranslator.cpp
--- a/Engine/Plugins/Interchange/Extensions/OpenUSD/Source/Import/Private/InterchangeUsdTranslator.cpp
+++ b/Engine/Plugins/Interchange/Extensions/OpenUSD/Source/Import/Private/InterchangeUsdTranslator.cpp
@@ -197,7 +197,7 @@ UInterchangeUsdTranslatorSettings::UInterchangeUsdTranslatorSettings()
 		  0.01,				   // MetersPerUnit
 		  EUsdUpAxis::ZAxis	   // UpAxis
 	  }
-	, PointInstancerCollapsing(EUsdPointInstancerCollapsing::DontCollapse)
+	, PointInstancerCollapsing(EUsdPointInstancerCollapsing::NoCollapsing)
 	, bUseSchemaForCollapsing(false)
 	, bUsePrimKindsForCollapsing(false)
 	, KindsToCollapse(static_cast<int32>(EUsdDefaultKind::Component | EUsdDefaultKind::Subcomponent))
diff --git a/Engine/Plugins/Interchange/Extensions/OpenUSD/Source/Import/Private/SchemaHandlers/PointInstancerSchemaHandler.cpp b/Engine/Plugins/Interchange/Extensions/OpenUSD/Source/Import/Private/SchemaHandlers/PointInstancerSchemaHandler.cpp
--- a/Engine/Plugins/Interchange/Extensions/OpenUSD/Source/Import/Private/SchemaHandlers/PointInstancerSchemaHandler.cpp
+++ b/Engine/Plugins/Interchange/Extensions/OpenUSD/Source/Import/Private/SchemaHandlers/PointInstancerSchemaHandler.cpp
@@ -403,7 +403,7 @@ namespace UE::Interchange::USD
 	{
 		// If we can't collapse point instancers then we already know for sure this can't be collapsed
 		UInterchangeUsdTranslatorSettings* Settings = UsdContext.GetTranslatorSettings();
-		if (Settings && Settings->PointInstancerCollapsing == EUsdPointInstancerCollapsing::DontCollapse)
+		if (Settings && Settings->PointInstancerCollapsing == EUsdPointInstancerCollapsing::NoCollapsing)
 		{
 			return false;
 		}
@@ -415,7 +415,7 @@ namespace UE::Interchange::USD
 	{
 		// If we can't collapse point instancers then we already know for sure this can't collapse
 		UInterchangeUsdTranslatorSettings* Settings = UsdContext.GetTranslatorSettings();
-		if (Settings && Settings->PointInstancerCollapsing == EUsdPointInstancerCollapsing::DontCollapse)
+		if (Settings && Settings->PointInstancerCollapsing == EUsdPointInstancerCollapsing::NoCollapsing)
 		{
 			return false;
 		}
diff --git a/Engine/Plugins/Interchange/Extensions/OpenUSD/Source/Import/Public/InterchangeUsdTranslator.h b/Engine/Plugins/Interchange/Extensions/OpenUSD/Source/Import/Public/InterchangeUsdTranslator.h
--- a/Engine/Plugins/Interchange/Extensions/OpenUSD/Source/Import/Public/InterchangeUsdTranslator.h
+++ b/Engine/Plugins/Interchange/Extensions/OpenUSD/Source/Import/Public/InterchangeUsdTranslator.h
@@ -50,7 +50,7 @@ UENUM(BlueprintType)
 enum class EUsdPointInstancerCollapsing : uint8
 {
 	/** Prevent collapsing the subtree the PointInstancer occupies */
-	DontCollapse,
+	NoCollapsing,
 
 	/** Collapse the point instancer along with the rest of the subtree, baking its instances into the produced StaticMesh */
 	CollapseAsStaticMesh
