USD Interchange: Allow specifying custom render contexts for some schema handlers

#Synced-CL 49920387

diff --git a/Engine/Plugins/Interchange/Extensions/OpenUSD/Source/Editor/Private/InterchangeUsdHandlerOrderWindow.cpp b/Engine/Plugins/Interchange/Extensions/OpenUSD/Source/Editor/Private/InterchangeUsdHandlerOrderWindow.cpp
--- a/Engine/Plugins/Interchange/Extensions/OpenUSD/Source/Editor/Private/InterchangeUsdHandlerOrderWindow.cpp
+++ b/Engine/Plugins/Interchange/Extensions/OpenUSD/Source/Editor/Private/InterchangeUsdHandlerOrderWindow.cpp
@@ -18,6 +18,7 @@
 #include "PropertyEditorModule.h"
 #include "Widgets/Input/SButton.h"
 #include "Widgets/Input/SCheckBox.h"
+#include "Widgets/Input/SEditableTextBox.h"
 #include "Widgets/SBoxPanel.h"
 #include "Widgets/SWindow.h"
 #include "Widgets/Views/STableRow.h"
@@ -70,6 +71,20 @@ namespace UE::UsdHandlerOrderWindow::Private
 		return RenderContexts;
 	}
 
+	// Swaps the 'universal' string for the empty string, if it exists in the list
+	TArray<FString> GetSDKFacingRenderContexts(TArray<FString> RenderContexts)
+	{
+		for (FString& RenderContext : RenderContexts)
+		{
+			if (RenderContext == UnrealIdentifiers::UniversalRenderContextDisplayString)
+			{
+				RenderContext = FString{};
+			}
+		}
+
+		return RenderContexts;
+	}
+
 	class SSchemaHandlerRow : public STableRow<TSharedPtr<FSchemaHandlerEntry>>
 	{
 	public:
@@ -100,6 +115,58 @@ namespace UE::UsdHandlerOrderWindow::Private
 
 			TreeItem = InTreeItem;
 
+			TSharedPtr<SWidget> RenderContextWidget = SNullWidget::NullWidget;
+			if (TreeItem && TreeItem->SchemaName == TEXT("Material"))
+			{
+				if (InTreeItem->bAllowCustomRenderContexts)
+				{
+					// For handlers that allow custom render contexts, we'll let people type them in manually
+					RenderContextWidget = SNew(SEditableTextBox)
+					.Cursor(EMouseCursor::Default)
+					.ToolTipText(LOCTEXT("RenderContext", "Comma-separated list of render contexts to try parsing, in order of priority (first item of the list will be checked first)\nIn this dialog, use 'universal' for the USD universal render context."))
+					.MinDesiredWidth(70.0f)
+					.Text_Lambda([this]() -> FText
+					{
+						if (TreeItem)
+						{
+							return FText::FromString(FString::Join(GetUserFacingRenderContexts(TreeItem->CustomRenderContexts), TEXT(", ")));
+						}
+
+						return {};
+					})
+					.Font(FAppStyle::GetFontStyle("PropertyWindow.NormalFont"))
+					.OnTextCommitted_Lambda([this](const FText& NewText, ETextCommit::Type CommitType)
+					{
+						if (!TreeItem || CommitType != ETextCommit::OnEnter)
+						{
+							return;
+						}
+
+						FString NewString = NewText.ToString();
+
+						TArray<FString> RenderContexts;
+						NewString.ParseIntoArray(RenderContexts, TEXT(","));
+
+						for (int32 Index = RenderContexts.Num() - 1; Index >= 0; --Index)
+						{
+							FString& RenderContext = RenderContexts[Index];
+							RenderContext.TrimStartAndEndInline();
+						}
+
+						TreeItem->CustomRenderContexts = GetSDKFacingRenderContexts(RenderContexts);
+					});
+				}
+				else
+				{
+					// For other handlers, just display the default render contexts as text
+					RenderContextWidget = SNew(STextBlock)
+						.Text(FText::FromString(FString::Join(GetUserFacingRenderContexts(TreeItem->DefaultRenderContexts), TEXT(", "))))
+						.ToolTipText(LOCTEXT("RenderContext", "List of render contexts the handler will try parsing, in order of priority (first item of the list will be checked first)\nIn this dialog 'universal' means the USD universal render context."))
+						.MinDesiredWidth(50.0f)
+						.Font(FAppStyle::GetFontStyle("PropertyWindow.NormalFont"));
+				}
+			}
+
 			SetRowContent(
 				SNew(SBorder)
 				.BorderImage_Lambda([this]()
@@ -189,12 +256,7 @@ namespace UE::UsdHandlerOrderWindow::Private
 					.VAlign(VAlign_Center)
 					.FillWidth(0.7)
 					[
-						SNew(STextBlock)
-						.Text(FText::FromString(FString::Join(GetUserFacingRenderContexts(TreeItem->DefaultRenderContexts), TEXT(", "))))
-						.ToolTipText(LOCTEXT("RenderContext", "List of render contexts the handler will try parsing, in order of priority (first item of the list will be checked first)"))
-						.Visibility(TreeItem && TreeItem->SchemaName == TEXT("Material") ? EVisibility::Visible : EVisibility::Hidden)
-						.MinDesiredWidth(50.0f)
-						.Font(FAppStyle::GetFontStyle("PropertyWindow.NormalFont"))
+						RenderContextWidget.ToSharedRef()
 					]
 
 					+SHorizontalBox::Slot()
@@ -328,8 +390,8 @@ bool SInterchangeUsdHandlerOrderWindow::ShowWindow(UInterchangeUsdTranslatorSett
 	const static FString HeightConfigName = TEXT("DialogHeight");
 
 	// Grab our desired window size from persistent storage, if we have any
-	int32 Width = 540;
-	int32 Height = 860;
+	int32 Width = 650;
+	int32 Height = 850;
 	{
 		FDisplayMetrics DisplayMetrics;
 		FSlateApplication::Get().GetDisplayMetrics(DisplayMetrics);
diff --git a/Engine/Plugins/Interchange/Extensions/OpenUSD/Source/Import/Private/SchemaHandlers/MaterialSchemaHandler.cpp b/Engine/Plugins/Interchange/Extensions/OpenUSD/Source/Import/Private/SchemaHandlers/MaterialSchemaHandler.cpp
--- a/Engine/Plugins/Interchange/Extensions/OpenUSD/Source/Import/Private/SchemaHandlers/MaterialSchemaHandler.cpp
+++ b/Engine/Plugins/Interchange/Extensions/OpenUSD/Source/Import/Private/SchemaHandlers/MaterialSchemaHandler.cpp
@@ -15,6 +15,26 @@ namespace UE::Interchange::USD
 		return SchemaName;
 	}
 
+	bool FMaterialSchemaHandler::AllowCustomRenderContexts() const
+	{
+		return false;
+	}
+
+	const TArray<FString>& FMaterialSchemaHandler::GetCustomRenderContexts() const
+	{
+		return CustomRenderContexts;
+	}
+
+	void FMaterialSchemaHandler::SetCustomRenderContexts(const TArray<FString>& RenderContexts)
+	{
+		if (!ensure(AllowCustomRenderContexts()))
+		{
+			return;
+		}
+
+		CustomRenderContexts = RenderContexts;
+	}
+
 	bool FMaterialSchemaHandler::CanHandlePrim(const UE::FUsdPrim& Prim, const UInterchangeUsdContext& UsdContext) const
 	{
 		TRACE_CPUPROFILER_EVENT_SCOPE(FMaterialSchemaHandler::CanHandlePrim)
diff --git a/Engine/Plugins/Interchange/Extensions/OpenUSD/Source/Import/Private/SchemaHandlers/SchemaHandlerRegistry.cpp b/Engine/Plugins/Interchange/Extensions/OpenUSD/Source/Import/Private/SchemaHandlers/SchemaHandlerRegistry.cpp
--- a/Engine/Plugins/Interchange/Extensions/OpenUSD/Source/Import/Private/SchemaHandlers/SchemaHandlerRegistry.cpp
+++ b/Engine/Plugins/Interchange/Extensions/OpenUSD/Source/Import/Private/SchemaHandlers/SchemaHandlerRegistry.cpp
@@ -56,9 +56,8 @@ namespace UE::Interchange::USD
 				continue;
 			}
 
-			// Apply (potentially modified) settings stored on the entries back to the handlers
-			TSharedRef<FSchemaHandler> NewHandler = (*FoundGenerator)();
-			NewHandler->SetEnabled(Entry.bEnabled);
+			// The generator will also apply our handler entry settings to the new handler instance
+			TSharedRef<FSchemaHandler> NewHandler = (*FoundGenerator)(Entry);
 
 			// We check for duplicate handler names here instead of in Registry because it's possible that user C++ code may have manipulated 
 			// the registered handler array directly, adding/removing handlers (as it was intended) and skipping our own Register function.
diff --git a/Engine/Plugins/Interchange/Extensions/OpenUSD/Source/Import/Private/SchemaHandlers/UniversalMaterialSchemaHandler.cpp b/Engine/Plugins/Interchange/Extensions/OpenUSD/Source/Import/Private/SchemaHandlers/UniversalMaterialSchemaHandler.cpp
--- a/Engine/Plugins/Interchange/Extensions/OpenUSD/Source/Import/Private/SchemaHandlers/UniversalMaterialSchemaHandler.cpp
+++ b/Engine/Plugins/Interchange/Extensions/OpenUSD/Source/Import/Private/SchemaHandlers/UniversalMaterialSchemaHandler.cpp
@@ -228,6 +228,7 @@ namespace UE::UniversalMaterialSchemaHandler::Private
 	void ConvertMaterialPrim(
 		const UE::FUsdPrim& Prim,
 		UInterchangeMaterialInstanceNode* MaterialNode,
+		const TArray<FString>& RenderContexts,
 		FHandlerAccumulatedInfo& AccumulatedInfo,
 		UInterchangeUsdContext& UsdContext
 	)
@@ -236,10 +237,23 @@ namespace UE::UniversalMaterialSchemaHandler::Private
 
 		UInterchangeBaseNodeContainer* NodeContainer = UsdContext.GetNodeContainer();
 
-		const static FString RenderContext = UnrealIdentifiers::UniversalRenderContext.ToString();
-
+		// Pick the first render context that we can parse this material with
+		bool bSuccess = false;
 		UsdToUnreal::FUsdPreviewSurfaceMaterialData MaterialData;
-		const bool bSuccess = UsdToUnreal::ConvertMaterial(Prim, MaterialData, *RenderContext);
+		for (const FString& RenderContext : RenderContexts)
+		{
+			UsdToUnreal::FUsdPreviewSurfaceMaterialData MaterialDataForRenderContext;
+			bSuccess = UsdToUnreal::ConvertMaterial(Prim, MaterialDataForRenderContext, *RenderContext);
+			if (bSuccess)
+			{
+				MaterialData = MaterialDataForRenderContext;
+				break;
+			}
+		}
+		if (!bSuccess)
+		{
+			return;
+		}
 
 		// Set all the parameter values to the interchange node
 		bool bNeedsVTParent = false;
@@ -333,6 +347,11 @@ namespace UE::Interchange::USD
 		return RenderContexts;
 	}
 
+	bool FUniversalMaterialSchemaHandler::AllowCustomRenderContexts() const
+	{
+		return true;
+	}
+
 	TOptional<bool> FUniversalMaterialSchemaHandler::CanBeCollapsed(const UE::FUsdPrim& Prim, UInterchangeUsdContext& UsdContext) const
 	{
 		return false;
@@ -394,7 +413,8 @@ namespace UE::Interchange::USD
 			return false;
 		}
 
-		ConvertMaterialPrim(Prim, AssetNode, AccumulatedInfo, UsdContext);
+		const TArray<FString>& RenderContexts = GetCustomRenderContexts();
+		ConvertMaterialPrim(Prim, AssetNode, RenderContexts, AccumulatedInfo, UsdContext);
 
 		return true;
 	}
diff --git a/Engine/Plugins/Interchange/Extensions/OpenUSD/Source/Import/Public/SchemaHandlers/MaterialSchemaHandler.h b/Engine/Plugins/Interchange/Extensions/OpenUSD/Source/Import/Public/SchemaHandlers/MaterialSchemaHandler.h
--- a/Engine/Plugins/Interchange/Extensions/OpenUSD/Source/Import/Public/SchemaHandlers/MaterialSchemaHandler.h
+++ b/Engine/Plugins/Interchange/Extensions/OpenUSD/Source/Import/Public/SchemaHandlers/MaterialSchemaHandler.h
@@ -21,7 +21,28 @@ namespace UE::Interchange::USD
 		 */
 		UE_API virtual const TArray<FString>& GetDefaultRenderContexts() const = 0;
 
+		/**
+		 * Returns whether this handler allows the user to provide custom render contexts.
+		 * This being false (the default value) means this handler is hard-coded to work with only specific render contexts.
+		 */
+		UE_API virtual bool AllowCustomRenderContexts() const;
+
+		/**
+		 * Returns the full list of render contexts that this handler will try parsing materials with.
+		 * This value is initialized with the output of GetDefaultRenderContexts(), so it should include those entries by default.
+		 */
+		UE_API virtual const TArray<FString>& GetCustomRenderContexts() const;
+
+		/**
+		 * Sets the full list of render contexts that this handler will try parsing materials with, if AllowCustomRenderContexts()
+		 * is enabled.
+		 */
+		UE_API virtual void SetCustomRenderContexts(const TArray<FString>& RenderContexts);
+
 		UE_API virtual bool CanHandlePrim(const UE::FUsdPrim& Prim, const UInterchangeUsdContext& UsdContext) const override;
+
+	protected:
+		TArray<FString> CustomRenderContexts;
 	};
 }
 
diff --git a/Engine/Plugins/Interchange/Extensions/OpenUSD/Source/Import/Public/SchemaHandlers/SchemaHandlerEntry.h b/Engine/Plugins/Interchange/Extensions/OpenUSD/Source/Import/Public/SchemaHandlers/SchemaHandlerEntry.h
--- a/Engine/Plugins/Interchange/Extensions/OpenUSD/Source/Import/Public/SchemaHandlers/SchemaHandlerEntry.h
+++ b/Engine/Plugins/Interchange/Extensions/OpenUSD/Source/Import/Public/SchemaHandlers/SchemaHandlerEntry.h
@@ -31,4 +31,19 @@ public:
 	/** For material schema handlers, we also hoist the render contexts the handler supports here, to be visible on the UI */
 	UPROPERTY(EditAnywhere, BlueprintReadOnly, Category = "USD Translator")
 	TArray<FString> DefaultRenderContexts;
+
+	/**
+	 * If true, it means the handler allows the user to manually specify which render contexts to parse (e.g. the default 'universal' handler does this).
+	 * If false, it means the handler is hard-coded to only look for a particular render context (e.g. the default MaterialX handler does this).
+	 */
+	UPROPERTY(EditAnywhere, BlueprintReadWrite, Category = "USD Translator")
+	bool bAllowCustomRenderContexts = false;
+
+	/**
+	 * The full list of custom render contexts that this handler should try parsing.
+	 * This is only used whenever bAllowCustomRenderContexts is enabled.
+	 * This value is initialized with the value of DefaultRenderContexts within FSchemaHandlerRegistry::Register.
+	 */
+	UPROPERTY(EditAnywhere, BlueprintReadWrite, Category = "USD Translator")
+	TArray<FString> CustomRenderContexts;
 };
diff --git a/Engine/Plugins/Interchange/Extensions/OpenUSD/Source/Import/Public/SchemaHandlers/SchemaHandlerRegistry.h b/Engine/Plugins/Interchange/Extensions/OpenUSD/Source/Import/Public/SchemaHandlers/SchemaHandlerRegistry.h
--- a/Engine/Plugins/Interchange/Extensions/OpenUSD/Source/Import/Public/SchemaHandlers/SchemaHandlerRegistry.h
+++ b/Engine/Plugins/Interchange/Extensions/OpenUSD/Source/Import/Public/SchemaHandlers/SchemaHandlerRegistry.h
@@ -15,7 +15,7 @@
 namespace UE::Interchange::USD
 {
 	/** Main lambda invoked when spawning an FSchemaHandler */
-	using FSchemaHandlerGenerator = TFunction<TSharedRef<FSchemaHandler>()>;
+	using FSchemaHandlerGenerator = TFunction<TSharedRef<FSchemaHandler>(const FSchemaHandlerEntry& Entry)>;
 
 	/**
 	 * Static storage of all registered FSchemaHandlers.
@@ -52,20 +52,47 @@ namespace UE::Interchange::USD
 			SchemaHandlerType TempHandler{};
 			const FString& HandlerName = TempHandler.GetHandlerName();
 
+			FSchemaHandlerGenerator GeneratorFunc;
+
 			FSchemaHandlerEntry& NewEntry = RegisteredHandlerEntries.Emplace_GetRef();
 			NewEntry.HandlerName = HandlerName;
 			NewEntry.bEnabled = TempHandler.IsEnabled();
 			NewEntry.SchemaName = TempHandler.GetTargetSchemaName();
 			if constexpr (std::is_base_of_v<FMaterialSchemaHandler, SchemaHandlerType>)
 			{
-				NewEntry.DefaultRenderContexts = static_cast<FMaterialSchemaHandler*>(&TempHandler)->GetDefaultRenderContexts();
+				FMaterialSchemaHandler* MaterialHandler = static_cast<FMaterialSchemaHandler*>(&TempHandler);
+				NewEntry.DefaultRenderContexts = MaterialHandler->GetDefaultRenderContexts();
+				NewEntry.bAllowCustomRenderContexts = MaterialHandler->AllowCustomRenderContexts();
+				if (MaterialHandler->AllowCustomRenderContexts())
+				{
+					// Initialize our custom render contexts with the same value as the default
+					NewEntry.CustomRenderContexts = NewEntry.DefaultRenderContexts;
+				}
+
+				GeneratorFunc = [](const FSchemaHandlerEntry& Entry)
+				{
+					TSharedRef<FSchemaHandler> Result = MakeShared<SchemaHandlerType>();
+					Result->SetEnabled(Entry.bEnabled);
+
+					TSharedRef<FMaterialSchemaHandler> MaterialHandlerResult = StaticCastSharedRef<FMaterialSchemaHandler>(Result);
+					if (MaterialHandlerResult->AllowCustomRenderContexts())
+					{
+						MaterialHandlerResult->SetCustomRenderContexts(Entry.CustomRenderContexts);
+					}
+
+					return Result;
+				};
 			}
-
-			FSchemaHandlerGenerator GeneratorFunc = []()
+			else
 			{
-				TSharedRef<FSchemaHandler> Result = MakeShared<SchemaHandlerType>();
-				return Result;
-			};
+				GeneratorFunc = [](const FSchemaHandlerEntry& Entry)
+				{
+					TSharedRef<FSchemaHandler> Result = MakeShared<SchemaHandlerType>();
+					Result->SetEnabled(Entry.bEnabled);
+					return Result;
+				};
+			}
+
 			RegisteredHandlerNamesToGenerators.Add(HandlerName, GeneratorFunc);
 
 			return HandlerName;
diff --git a/Engine/Plugins/Interchange/Extensions/OpenUSD/Source/Import/Public/SchemaHandlers/UniversalMaterialSchemaHandler.h b/Engine/Plugins/Interchange/Extensions/OpenUSD/Source/Import/Public/SchemaHandlers/UniversalMaterialSchemaHandler.h
--- a/Engine/Plugins/Interchange/Extensions/OpenUSD/Source/Import/Public/SchemaHandlers/UniversalMaterialSchemaHandler.h
+++ b/Engine/Plugins/Interchange/Extensions/OpenUSD/Source/Import/Public/SchemaHandlers/UniversalMaterialSchemaHandler.h
@@ -17,6 +17,8 @@ namespace UE::Interchange::USD
 
 		UE_API const TArray<FString>& GetDefaultRenderContexts() const override;
 
+		UE_API virtual bool AllowCustomRenderContexts() const override;
+
 		UE_API virtual TOptional<bool> CanBeCollapsed(const UE::FUsdPrim& Prim, UInterchangeUsdContext& UsdContext) const override;
 
 		UE_API virtual TOptional<bool> CollapsesChildren(const UE::FUsdPrim& Prim, UInterchangeUsdContext& UsdContext) const override;
