Interchange USD: Fix bad mesh descriptions when all meshes from a bundle were disabled via purpose.

Update automated tests to the latest changes related to geometry purpose and point instancer collapsing.

#Synced-CL 49920387

diff --git a/Engine/Plugins/Interchange/Extensions/OpenUSD/Source/Import/Private/InterchangeUsdTranslator.cpp b/Engine/Plugins/Interchange/Extensions/OpenUSD/Source/Import/Private/InterchangeUsdTranslator.cpp
--- a/Engine/Plugins/Interchange/Extensions/OpenUSD/Source/Import/Private/InterchangeUsdTranslator.cpp
+++ b/Engine/Plugins/Interchange/Extensions/OpenUSD/Source/Import/Private/InterchangeUsdTranslator.cpp
@@ -281,28 +281,39 @@ TOptional<UE::Interchange::USD::FHandlerAccumulatedInfo> UInterchangeUSDTranslat
 	// reallocate, it could invalidate the reference
 	UE::Interchange::USD::FHandlerAccumulatedInfo AccumulatedInfo = UsdContext->HandledPrimInfo.Emplace(PrimPath);
 	{
-		// By specifying a nullptr scene node here, any call to AccumulatedInfo.GetMainSceneNode() will return nullptr.
+		// By temporarily pushing nullptr scene node here, any call to AccumulatedInfo.GetMainSceneNode() will return nullptr.
 		// This is sort of a trick to prevent creating scene nodes internally if the caller is just interested on the produced
 		// asset nodes
+		bool bInsertedNullptrSceneNode = false;
 		if (!bAllowSceneNodeGeneration)
 		{
-			AccumulatedInfo.PrimSceneNodes.Add(nullptr);
+			if (AccumulatedInfo.PrimSceneNodes.Num() == 0 || AccumulatedInfo.PrimSceneNodes[0] != nullptr)
+			{
+				AccumulatedInfo.PrimSceneNodes.Insert(nullptr, 0);
+				bInsertedNullptrSceneNode = true;
+			}
 		}
 
 		// For now just emit a single scene node for the pseudoroot
 		if (Prim.IsPseudoRoot())
 		{
-			UInterchangeBaseNode* RootNode = AccumulatedInfo.GetOrCreateMainSceneNode(Prim, TraversalInfo, *UsdContext);
-			return AccumulatedInfo;
+			AccumulatedInfo.GetOrCreateMainSceneNode(Prim, TraversalInfo, *UsdContext);
 		}
-
-		for (const TSharedRef<FSchemaHandler>& Handler : ImplPtr->SchemaHandlers)
+		else
 		{
-			if (Handler->IsEnabled() && Handler->CanHandlePrim(Prim, *UsdContext))
+			for (const TSharedRef<FSchemaHandler>& Handler : ImplPtr->SchemaHandlers)
 			{
-				Handler->OnTranslate(Prim, TraversalInfo, AccumulatedInfo, *UsdContext);
+				if (Handler->IsEnabled() && Handler->CanHandlePrim(Prim, *UsdContext))
+				{
+					Handler->OnTranslate(Prim, TraversalInfo, AccumulatedInfo, *UsdContext);
+				}
 			}
 		}
+
+		if (bInsertedNullptrSceneNode && AccumulatedInfo.PrimSceneNodes.Num() > 0)
+		{
+			AccumulatedInfo.PrimSceneNodes.RemoveAt(0);
+		}
 	}
 	UsdContext->HandledPrimInfo.Add(PrimPath, AccumulatedInfo);
 
diff --git a/Engine/Plugins/Interchange/Extensions/OpenUSD/Source/Import/Private/SchemaHandlers/CollapsingSchemaHandler.cpp b/Engine/Plugins/Interchange/Extensions/OpenUSD/Source/Import/Private/SchemaHandlers/CollapsingSchemaHandler.cpp
--- a/Engine/Plugins/Interchange/Extensions/OpenUSD/Source/Import/Private/SchemaHandlers/CollapsingSchemaHandler.cpp
+++ b/Engine/Plugins/Interchange/Extensions/OpenUSD/Source/Import/Private/SchemaHandlers/CollapsingSchemaHandler.cpp
@@ -206,49 +206,50 @@ namespace UE::Interchange::USD
 
 		using namespace UE::CollapsingSchemaHandler::Private;
 
-		// If we're a collapsing root, create our mesh bundle
-		UInterchangeMeshBundleNode* Bundle = nullptr;
-		if (FInterchangeUsdInfoCache* InfoCache = UsdContext.GetInterchangeInfoCache())
+		// If we're already collapsing a mesh we don't want to create a nested mesh bundle and want instead to treat this as
+		// a regular prim
+		if (TraversalInfo.bCurrentlyCollapsing)
 		{
-			if (InfoCache->DoesPathCollapseChildren(Prim.GetPrimPath()) && TraversalInfo.CurrentCollapsedMesh == nullptr)
-			{
-				UE::FUsdPrim PrototypePrim = UsdUtils::GetPrototypePrim(Prim);
-				const FString PrimPath = PrototypePrim.GetPrimPath().GetString();
-
-				// Don't use the prototype prim for the name though, as we don't want our asset named __Prototype_N
-				// This is only relevant for collapsed meshes and not other asset types because this is the only case where the
-				// asset gets named after the topmost prim (i.e. the instance/prototype prim itself)
-				FString NewNodeName{Prim.GetName().ToString()}; 
-				FString NewNodeUid = MakeNodeUid(MeshPrefix + PrimPath);
-
-				Bundle = AccumulatedInfo.GetOrCreateAssetNode<UInterchangeMeshBundleNode>(*UsdContext.GetNodeContainer(), NewNodeUid, NewNodeName);
-			}
+			return false;
 		}
 
-		if (Bundle)
+		// We can early out if our prim doesn't collapse anyway
+		if (FInterchangeUsdInfoCache* InfoCache = UsdContext.GetInterchangeInfoCache())
 		{
-			// Only ever setup our scene node if we can collapse stuff. Otherwise we can just yield the creation of the scene
-			// node to the other schema handlers
-			UInterchangeBaseNode* SceneNodeBase = AccumulatedInfo.GetOrCreateMainSceneNode(Prim, TraversalInfo, UsdContext);
-			if (UInterchangeSceneNode* SceneNode = Cast<UInterchangeSceneNode>(SceneNodeBase))
+			if (!InfoCache->DoesPathCollapseChildren(Prim.GetPrimPath()))
 			{
-				SceneNode->SetCustomAssetInstanceUid(Bundle->GetUniqueID());
+				return false;
 			}
+		}
 
-			TraversalInfo.CurrentCollapsedMesh = Bundle;
+		// Try collapsing the hierarchy to see if it produces any collapsed mesh nodes
+		TMap<FString, TArray<FTransform>> CurrentMeshToTransforms;
+		TArray<FUsdPrim> Children;
+		{
+			TGuardValue<bool> CollapsingGuard{TraversalInfo.bCurrentlyCollapsing, true};
 
-			// Build up our map in here and serialize it into the bundle node only once
-			TMap<FString, TArray<FTransform>> CurrentMeshToTransforms;
-			{
-				TArray<FTransform> TransformsToMeshBundleRoot{FTransform::Identity};
+			TArray<FTransform> TransformsToMeshBundleRoot{FTransform::Identity};
 
-				constexpr bool bTraverseInstanceProxies = true;
-				const TArray<FUsdPrim> Children = Prim.GetFilteredChildren(bTraverseInstanceProxies);
+			constexpr bool bTraverseInstanceProxies = true;
+			Children = Prim.GetFilteredChildren(bTraverseInstanceProxies);
 
-				if (Prim.IsA(TEXT("PointInstancer")))
+			if (Prim.IsA(TEXT("PointInstancer")))
+			{
+				RecursiveCollapsePointInstancer(
+					Prim,
+					TraversalInfo,
+					TransformsToMeshBundleRoot,
+					CurrentMeshToTransforms,
+					AccumulatedInfo,
+					UsdContext
+				);
+			}
+			else
+			{
+				for (const FUsdPrim& ChildPrim : Children)
 				{
-					RecursiveCollapsePointInstancer(
-						Prim,
+					RecursiveCollapseSceneHierarchy(
+						ChildPrim,
 						TraversalInfo,
 						TransformsToMeshBundleRoot,
 						CurrentMeshToTransforms,
@@ -256,30 +257,64 @@ namespace UE::Interchange::USD
 						UsdContext
 					);
 				}
-				else
-				{
-					for (const FUsdPrim& ChildPrim : Children)
-					{
-						RecursiveCollapseSceneHierarchy(
-							ChildPrim,
-							TraversalInfo,
-							TransformsToMeshBundleRoot,
-							CurrentMeshToTransforms,
-							AccumulatedInfo,
-							UsdContext
-						);
-					}
-				}
+			}
+		}
 
-				// Make sure all of our child prims are marked as visited however, even if our recursive functions potentially ignored
-				// them, as otherwise the regular translation process will continue stepping into them. Maybe in the future this could be
-				// useful, so that we could "partially collapse" a hierarchy. For now it's simpler to just avoid this case however
-				for (const FUsdPrim& ChildPrim : Children)
-				{
-					UsdContext.HandledPrimInfo.FindOrAdd(ChildPrim.GetPrimPath());
-				}
+		// Only emit a mesh bundle node if we have at least one mesh node in there.
+		//
+		// That should be the correct thing to do in principle, but even in practice we should do this
+		// because if we end up creating a factory node from an empty mesh bundle we could run into trouble
+		// when we later have no actual payload data to add to the produced asset, and it's left in some
+		// incomplete state.
+		bool bAtLeastOneMeshInstance = false;
+		for (const TPair<FString, TArray<FTransform>>& MeshTransforms : CurrentMeshToTransforms)
+		{
+			if (MeshTransforms.Value.Num() > 0)
+			{
+				bAtLeastOneMeshInstance = true;
+				break;
 			}
+		}
+		if (!bAtLeastOneMeshInstance)
+		{
+			return false;
+		}
+
+		// OK, we're collapsing this subtree. Make sure all of our child prims are marked as visited, even if our recursive functions
+		// potentially ignored them, as otherwise the regular translation process will continue stepping into them. Maybe in the
+		// future this could be useful, so that we could "partially collapse" a hierarchy. For now it's simpler to just avoid
+		// this case however
+		for (const FUsdPrim& ChildPrim : Children)
+		{
+			UsdContext.HandledPrimInfo.FindOrAdd(ChildPrim.GetPrimPath());
+		}
+
+		UE::FUsdPrim PrototypePrim = UsdUtils::GetPrototypePrim(Prim);
+		const FString PrimPath = PrototypePrim.GetPrimPath().GetString();
+
+		// Don't use the prototype prim for the name though, as we don't want our asset named __Prototype_N
+		// This is only relevant for collapsed meshes and not other asset types because this is the only case where the
+		// asset gets named after the topmost prim (i.e. the instance/prototype prim itself)
+		FString NewNodeName{Prim.GetName().ToString()};
+		FString NewNodeUid = MakeNodeUid(MeshPrefix + PrimPath);
+
+		UInterchangeMeshBundleNode* Bundle = AccumulatedInfo.GetOrCreateAssetNode<UInterchangeMeshBundleNode>(
+			*UsdContext.GetNodeContainer(),
+			NewNodeUid,
+			NewNodeName
+		);
+
+		if (Bundle)
+		{
 			Bundle->SetAllMeshNodesAndTransforms(CurrentMeshToTransforms);
+
+			// Only ever setup our scene node if we can collapse stuff. Otherwise we can just yield the creation of the scene
+			// node to the other schema handlers
+			UInterchangeBaseNode* SceneNodeBase = AccumulatedInfo.GetOrCreateMainSceneNode(Prim, TraversalInfo, UsdContext);
+			if (UInterchangeSceneNode* SceneNode = Cast<UInterchangeSceneNode>(SceneNodeBase))
+			{
+				SceneNode->SetCustomAssetInstanceUid(Bundle->GetUniqueID());
+			}
 		}
 
 		return true;
diff --git a/Engine/Plugins/Interchange/Extensions/OpenUSD/Source/Import/Public/InterchangeUsdTraversalInfo.h b/Engine/Plugins/Interchange/Extensions/OpenUSD/Source/Import/Public/InterchangeUsdTraversalInfo.h
--- a/Engine/Plugins/Interchange/Extensions/OpenUSD/Source/Import/Public/InterchangeUsdTraversalInfo.h
+++ b/Engine/Plugins/Interchange/Extensions/OpenUSD/Source/Import/Public/InterchangeUsdTraversalInfo.h
@@ -35,8 +35,8 @@ namespace UE::Interchange::USD
 
 		FTransform SceneGlobalTransform;
 
-		// If we're currently collapsing a hierarchy, we'll stash the mesh bundle node we're using in here
-		UInterchangeMeshBundleNode* CurrentCollapsedMesh = nullptr;
+		// Whether we're already collapsing a subtree or not (used to prevent us from creating nested mesh bundles)
+		bool bCurrentlyCollapsing = false;
 
 		TSharedPtr<UE::FUsdSkelCache> FurthestSkelCache;
 		TSharedPtr<FString> FurthestParentSkelRootPath;	   // Used to populate the skel cache
