USD: Fix collision on LOCTEXT macro key.

#Synced-CL 49920387

diff --git a/Engine/Plugins/Interchange/Extensions/OpenUSD/Source/Editor/Private/InterchangeUsdHandlerOrderWindow.cpp b/Engine/Plugins/Interchange/Extensions/OpenUSD/Source/Editor/Private/InterchangeUsdHandlerOrderWindow.cpp
--- a/Engine/Plugins/Interchange/Extensions/OpenUSD/Source/Editor/Private/InterchangeUsdHandlerOrderWindow.cpp
+++ b/Engine/Plugins/Interchange/Extensions/OpenUSD/Source/Editor/Private/InterchangeUsdHandlerOrderWindow.cpp
@@ -123,7 +123,7 @@ namespace UE::UsdHandlerOrderWindow::Private
 					// For handlers that allow custom render contexts, we'll let people type them in manually
 					RenderContextWidget = SNew(SEditableTextBox)
 					.Cursor(EMouseCursor::Default)
-					.ToolTipText(LOCTEXT("RenderContext", "Comma-separated list of render contexts to try parsing, in order of priority (first item of the list will be checked first)\nIn this dialog, use 'universal' for the USD universal render context."))
+					.ToolTipText(LOCTEXT("EditableRenderContextToolTip", "Comma-separated list of render contexts to try parsing, in order of priority (first item of the list will be checked first)\nIn this dialog, use 'universal' for the USD universal render context."))
 					.MinDesiredWidth(70.0f)
 					.Text_Lambda([this]() -> FText
 					{
@@ -161,7 +161,7 @@ namespace UE::UsdHandlerOrderWindow::Private
 					// For other handlers, just display the default render contexts as text
 					RenderContextWidget = SNew(STextBlock)
 						.Text(FText::FromString(FString::Join(GetUserFacingRenderContexts(TreeItem->DefaultRenderContexts), TEXT(", "))))
-						.ToolTipText(LOCTEXT("RenderContext", "List of render contexts the handler will try parsing, in order of priority (first item of the list will be checked first)\nIn this dialog 'universal' means the USD universal render context."))
+						.ToolTipText(LOCTEXT("RenderContextToolTip", "List of render contexts the handler will try parsing, in order of priority (first item of the list will be checked first)\nIn this dialog 'universal' means the USD universal render context."))
 						.MinDesiredWidth(50.0f)
 						.Font(FAppStyle::GetFontStyle("PropertyWindow.NormalFont"));
 				}
