USD: Fix support for deprecated unrealMaterial string attributes

#Synced-CL 50957406

diff --git a/Engine/Plugins/Interchange/Extensions/OpenUSD/Source/Import/Private/SchemaHandlers/GprimSchemaHandler.cpp b/Engine/Plugins/Interchange/Extensions/OpenUSD/Source/Import/Private/SchemaHandlers/GprimSchemaHandler.cpp
--- a/Engine/Plugins/Interchange/Extensions/OpenUSD/Source/Import/Private/SchemaHandlers/GprimSchemaHandler.cpp
+++ b/Engine/Plugins/Interchange/Extensions/OpenUSD/Source/Import/Private/SchemaHandlers/GprimSchemaHandler.cpp
@@ -276,14 +276,29 @@ namespace UE::GeomGprimSchemaHandler::Private
 
 		// Material assignments
 		{
+			// Note that here we only care about finding the materials that are bound for our slot dependencies, and the render context
+			// we provide to this function doesn't affect that all: We'll find the the same bound material prim regardless of which we use.
+			// By passing the universal render context here though, we ensure we *check* for UnrealMaterial prims, but don't actually "resolve"
+			// them to the content path just yet (which is what would happen if we provided the "unreal" render context): We want to just
+			// be pointing at the actual Material prims here, so that our SetSlotMaterialDependencies() correctly points at the material node
+			// we'll produce when handling the prim, depending on the render context schema handlers.
+			//
+			// The only annoyance left is we must provide bForceCheckUnrealMaterialAttribute=true here for backwards compatibility with the
+			// "unrealMaterial" string attribute: There is no Material prim in that case at all (so nothing for the schema handlers to be
+			// invoked on), and if we do not provide the argument we'll fully miss out on these deprecated material bindings.
+			// We can't provide the "unreal" render context for that, because that will cause us to early resolve those bindings to the 
+			// UAsset content paths directly, and as said above we need to be pointing at the prim to let the material schema handlers do that
 			double TimeCode = UsdUtils::GetDefaultTimeCode();
 			const bool bProvideMaterialIndices = false;
+			const pxr::TfToken RenderContext = pxr::UsdShadeTokens->universalRenderContext;
+			const bool bForceCheckUnrealMaterialAttribute = true;
 			UsdUtils::FUsdPrimMaterialAssignmentInfo PrototypeAssignments = UsdUtils::GetPrimMaterialAssignments(
 				MeshPrim,
 				TimeCode,
 				bProvideMaterialIndices,
-				UsdContext.CachedMeshConversionOptions.RenderContext,
-				UsdContext.CachedMeshConversionOptions.MaterialPurpose
+				RenderContext,
+				UsdContext.CachedMeshConversionOptions.MaterialPurpose,
+				bForceCheckUnrealMaterialAttribute
 			);
 
 			if (Info.bInsideLODVariant)
@@ -738,14 +753,19 @@ namespace UE::Interchange::USD
 			if (Prim.IsInstance() || Prim.IsInstanceProxy())
 			{
 				// This is in charge of picking up on collection-based material assignments, as we'll use the instance proxy prim here
+				//
+				// See the large comment on the other call to GetPrimMaterialAssignments, from within FGprimSchemaHandler::OnTranslate()
 				const double TimeCode = UsdUtils::GetDefaultTimeCode();
 				const bool bProvideMaterialIndices = false;
+				const pxr::TfToken RenderContext = pxr::UsdShadeTokens->universalRenderContext;
+				const bool bForceCheckUnrealMaterialAttribute = true;
 				UsdUtils::FUsdPrimMaterialAssignmentInfo InstanceAssignments = UsdUtils::GetPrimMaterialAssignments(
 					Prim,
 					TimeCode,
 					bProvideMaterialIndices,
-					UsdContext.CachedMeshConversionOptions.RenderContext,
-					UsdContext.CachedMeshConversionOptions.MaterialPurpose
+					RenderContext,
+					UsdContext.CachedMeshConversionOptions.MaterialPurpose,
+					bForceCheckUnrealMaterialAttribute
 				);
 
 				SetSlotMaterialDependencies(SceneNode, InstanceAssignments, AccumulatedInfo, UsdContext);
diff --git a/Engine/Plugins/Runtime/USDCore/Source/USDUtilities/Private/USDGeomMeshConversion.cpp b/Engine/Plugins/Runtime/USDCore/Source/USDUtilities/Private/USDGeomMeshConversion.cpp
--- a/Engine/Plugins/Runtime/USDCore/Source/USDUtilities/Private/USDGeomMeshConversion.cpp
+++ b/Engine/Plugins/Runtime/USDCore/Source/USDUtilities/Private/USDGeomMeshConversion.cpp
@@ -4212,7 +4212,8 @@ UsdUtils::FUsdPrimMaterialAssignmentInfo UsdUtils::GetPrimMaterialAssignments(
 	const pxr::UsdTimeCode TimeCode,
 	bool bProvideMaterialIndices,
 	const pxr::TfToken& RenderContext,
-	const pxr::TfToken& MaterialPurpose
+	const pxr::TfToken& MaterialPurpose,
+	bool bForceCheckUnrealMaterialAttribute
 )
 {
 	TRACE_CPUPROFILER_EVENT_SCOPE(UsdUtils::GetPrimMaterialAssignments);
@@ -4345,9 +4346,12 @@ UsdUtils::FUsdPrimMaterialAssignmentInfo UsdUtils::GetPrimMaterialAssignments(
 						}
 					}
 				}
+			}
 
-				// Priority 0.1.2: Partitition has an unrealMaterial attribute directly on it
-				if (!bHasAssignment)
+			// Priority 0.1.2: Partitition has an unrealMaterial attribute directly on it
+			if (!bHasAssignment)
+			{
+				if (bForceCheckUnrealMaterialAttribute || RenderContext == UnrealIdentifiers::Unreal)
 				{
 					if (TOptional<FString> UnrealMaterial = FetchFirstUEMaterialFromAttribute(GeomSubsetPrim, TimeCode))
 					{
@@ -4475,9 +4479,12 @@ UsdUtils::FUsdPrimMaterialAssignmentInfo UsdUtils::GetPrimMaterialAssignments(
 					bHasMainAssignment = true;
 				}
 			}
+		}
 
-			// Priority 1.2: unrealMaterial attribute directly on the prim
-			if (!bHasMainAssignment)
+		// Priority 1.2: unrealMaterial attribute directly on the prim
+		if (!bHasMainAssignment)
+		{
+			if (bForceCheckUnrealMaterialAttribute || RenderContext == UnrealIdentifiers::Unreal)
 			{
 				if (TOptional<FString> UnrealMaterial = FetchFirstUEMaterialFromAttribute(UsdPrim, TimeCode))
 				{
diff --git a/Engine/Plugins/Runtime/USDCore/Source/USDUtilities/Public/USDGeomMeshConversion.h b/Engine/Plugins/Runtime/USDCore/Source/USDUtilities/Public/USDGeomMeshConversion.h
--- a/Engine/Plugins/Runtime/USDCore/Source/USDUtilities/Public/USDGeomMeshConversion.h
+++ b/Engine/Plugins/Runtime/USDCore/Source/USDUtilities/Public/USDGeomMeshConversion.h
@@ -413,6 +413,7 @@ namespace UsdUtils
 	 * is false, MaterialIndices on the result Struct will have zero values
 	 * @param RenderContext - Which render context to get the materials for. Defaults to universal.
 	 * @param MaterialPurpose - Which material purpose to use when retrieving material bindings
+	 * @param bForceCheckUnrealMaterialAttribute - Whether to always check the "unrealMaterial" string attribute regardless of render context
 	 * @return Struct containing an array of material assignments, and a corresponding array of material indices for all polygons of the prim,
 	 * matching the ordering of the material assignments.
 	 */
@@ -421,7 +422,8 @@ namespace UsdUtils
 		const pxr::UsdTimeCode TimeCode = pxr::UsdTimeCode::EarliestTime(),
 		bool bProvideMaterialIndices = true,
 		const pxr::TfToken& RenderContext = pxr::UsdShadeTokens->universalRenderContext,
-		const pxr::TfToken& MaterialPurpose = pxr::UsdShadeTokens->allPurpose
+		const pxr::TfToken& MaterialPurpose = pxr::UsdShadeTokens->allPurpose,
+		bool bForceCheckUnrealMaterialAttribute = false
 	);
 
 	/** Returns whether this prim can be interpreted as describing a static mesh with multiple LODs (i.e. if the prim holds the LOD variant set
