USD: Fix universal render context handler no longer acting as a fallback as it only work when finding a render context it could parse

#Synced-CL 50957406

diff --git a/Engine/Plugins/Interchange/Extensions/OpenUSD/Source/Import/Private/SchemaHandlers/UniversalMaterialSchemaHandler.cpp b/Engine/Plugins/Interchange/Extensions/OpenUSD/Source/Import/Private/SchemaHandlers/UniversalMaterialSchemaHandler.cpp
--- a/Engine/Plugins/Interchange/Extensions/OpenUSD/Source/Import/Private/SchemaHandlers/UniversalMaterialSchemaHandler.cpp
+++ b/Engine/Plugins/Interchange/Extensions/OpenUSD/Source/Import/Private/SchemaHandlers/UniversalMaterialSchemaHandler.cpp
@@ -241,23 +241,20 @@ namespace UE::UniversalMaterialSchemaHandler::Private
 
 		UInterchangeBaseNodeContainer* NodeContainer = UsdContext.GetNodeContainer();
 
-		// Pick the first render context that we can parse this material with
-		bool bSuccess = false;
+		// Try parsing the material with all the render contexts we support, in order.
+		// Even if we find nothing that works we must continue however, so as to finish setting up our
+		// MaterialNode with a valid (default) CustomParent at the bottom of this function.
+		// Otherwise it's possible we end up producing no valid factory nodes for this USD material
 		UsdToUnreal::FUsdPreviewSurfaceMaterialData MaterialData;
 		for (const FString& RenderContext : RenderContexts)
 		{
 			UsdToUnreal::FUsdPreviewSurfaceMaterialData MaterialDataForRenderContext;
-			bSuccess = UsdToUnreal::ConvertMaterial(Prim, MaterialDataForRenderContext, *RenderContext);
-			if (bSuccess)
+			if (UsdToUnreal::ConvertMaterial(Prim, MaterialDataForRenderContext, *RenderContext))
 			{
 				MaterialData = MaterialDataForRenderContext;
 				break;
 			}
 		}
-		if (!bSuccess)
-		{
-			return;
-		}
 
 		// Set all the parameter values to the interchange node
 		bool bNeedsVTParent = false;
