USD Interchange: Reuse the same texture nodes if they use the same texture file and UTexture settings.

#Synced-CL 50957406

diff --git a/Engine/Plugins/Interchange/Extensions/OpenUSD/Source/Import/Private/SchemaHandlers/UniversalMaterialSchemaHandler.cpp b/Engine/Plugins/Interchange/Extensions/OpenUSD/Source/Import/Private/SchemaHandlers/UniversalMaterialSchemaHandler.cpp
--- a/Engine/Plugins/Interchange/Extensions/OpenUSD/Source/Import/Private/SchemaHandlers/UniversalMaterialSchemaHandler.cpp
+++ b/Engine/Plugins/Interchange/Extensions/OpenUSD/Source/Import/Private/SchemaHandlers/UniversalMaterialSchemaHandler.cpp
@@ -30,16 +30,26 @@ namespace UE::UniversalMaterialSchemaHandler::Private
 {
 	using namespace UE::Interchange::USD;
 
-	void CreateTextureNode(
-		const UE::FUsdPrim& Prim,
-		const FString& NodeUid,
+	FString CreateTextureNode(
 		const UsdToUnreal::FTextureParameterValue& Value,
 		FHandlerAccumulatedInfo& AccumulatedInfo,
 		UInterchangeBaseNodeContainer& NodeContainer
 	)
 	{
-		const FString PrimPath = UsdUtils::GetPrototypePrimPath(Prim).GetString();
-		const FString NodeName{FPaths::GetBaseFilename(Value.TextureFilePath)};
+		FString NodeName;
+		FString NodeUid;
+		const bool bSuccess = UInterchangeTextureNode::ExtractNodeUidAndNameFromFilePath(Value.TextureFilePath, NodeUid, NodeName);
+		if (!bSuccess)
+		{
+			return {};
+		}
+
+		// We can reuse the same texture node for multiple material channels, but only if they require the exact same settings.
+		// This means we must encode the final requested texture settings into the NodeUID and deduplicate using those.
+		// A texture used as a normal map needs normal map compression for example, and we may have multiple texture sample nodes pointing at
+		// the same filepath and claiming different sRGB values (It's not clear whether that would be a correct thing to do but it's something
+		// that could be done)
+		NodeUid = FString::Printf(TEXT("%s_%d_%d_%d_%d_%d"), *NodeUid, Value.Group, Value.GetSRGBValue(), Value.AddressX, Value.AddressY, Value.bIsUDIM);
 
 		UInterchangeTexture2DNode* Node = GetExistingNode<UInterchangeTexture2DNode>(NodeContainer, NodeUid);
 		if (!Node)
@@ -80,6 +90,8 @@ namespace UE::UniversalMaterialSchemaHandler::Private
 		}
 
 		AccumulatedInfo.PrimAssetNodes.Add(Node);
+
+		return NodeUid;
 	}
 
 	// We use this visitor to set UsdToUnreal::FParameterValue TVariant values onto UInterchangeMaterialInstanceNode.
@@ -125,13 +137,11 @@ namespace UE::UniversalMaterialSchemaHandler::Private
 
 		void operator()(const UsdToUnreal::FTextureParameterValue& Value) const
 		{
-			// Emit texture node itself (this is the main place where this happens)
-			// Note that the node name isn't just the texture path, as we may have multiple material users of this texture
-			// with different settings, and so we need separate translated nodes for each material and parameter
-			// TODO: Encode the different settings on the node uid as well to automatically deduplicate?
-			const FString
-				TextureUid = FString::Printf(TEXT("%s%s:%s"), *TexturePrefix, *UsdUtils::GetPrototypePrimPath(Prim).GetString(), **BaseParameterName);
-			CreateTextureNode(Prim, TextureUid, Value, AccumulatedInfo, NodeContainer);
+			const FString TextureUid = CreateTextureNode(Value, AccumulatedInfo, NodeContainer);
+			if (TextureUid.IsEmpty())
+			{
+				return;
+			}
 
 			// Actual texture assignment
 			MaterialNode.AddTextureParameterValue(	  //
