USD: Change "Material purpose" string to "Material Purpose" on the translator settings, as that was the only setting not fully upper case

#Synced-CL 51261034

diff --git a/Engine/Plugins/Interchange/Extensions/OpenUSD/Source/Editor/Private/InterchangeUsdTranslatorSettingsCustomization.cpp b/Engine/Plugins/Interchange/Extensions/OpenUSD/Source/Editor/Private/InterchangeUsdTranslatorSettingsCustomization.cpp
--- a/Engine/Plugins/Interchange/Extensions/OpenUSD/Source/Editor/Private/InterchangeUsdTranslatorSettingsCustomization.cpp
+++ b/Engine/Plugins/Interchange/Extensions/OpenUSD/Source/Editor/Private/InterchangeUsdTranslatorSettingsCustomization.cpp
@@ -106,7 +106,7 @@ void FInterchangeUsdTranslatorSettingsCustomization::CustomizeDetails(IDetailLay
 		.NameContent()
 		[
 			SNew(STextBlock)
-			.Text(FText::FromString(TEXT("Material purpose")))
+			.Text(FText::FromString(TEXT("Material Purpose")))
 			.Font(FAppStyle::GetFontStyle(TEXT("PropertyWindow.NormalFont")))
 			.ToolTipText(MaterialPurposeProperty->GetToolTipText())
 		]
