USD: Fix legacy importer no longer falling back to displayColor when bound materials are missing the required render context.

#Synced-CL 51330716

diff --git a/Engine/Plugins/Interchange/Extensions/OpenUSD/Source/Import/Private/SchemaHandlers/GprimSchemaHandler.cpp b/Engine/Plugins/Interchange/Extensions/OpenUSD/Source/Import/Private/SchemaHandlers/GprimSchemaHandler.cpp
--- a/Engine/Plugins/Interchange/Extensions/OpenUSD/Source/Import/Private/SchemaHandlers/GprimSchemaHandler.cpp
+++ b/Engine/Plugins/Interchange/Extensions/OpenUSD/Source/Import/Private/SchemaHandlers/GprimSchemaHandler.cpp
@@ -285,20 +285,22 @@ namespace UE::GeomGprimSchemaHandler::Private
 			//
 			// The only annoyance left is we must provide bForceCheckUnrealMaterialAttribute=true here for backwards compatibility with the
 			// "unrealMaterial" string attribute: There is no Material prim in that case at all (so nothing for the schema handlers to be
-			// invoked on), and if we do not provide the argument we'll fully miss out on these deprecated material bindings.
-			// We can't provide the "unreal" render context for that, because that will cause us to early resolve those bindings to the 
+			// invoked on), and if we do not provide the argument we'll fully miss out on these deprecated material bindings
+			// We can't provide the "unreal" render context for that, because that will cause us to early resolve those bindings to the
 			// UAsset content paths directly, and as said above we need to be pointing at the prim to let the material schema handlers do that
 			double TimeCode = UsdUtils::GetDefaultTimeCode();
 			const bool bProvideMaterialIndices = false;
 			const pxr::TfToken RenderContext = pxr::UsdShadeTokens->universalRenderContext;
 			const bool bForceCheckUnrealMaterialAttribute = true;
+			const bool bRequireMaterialsHaveProvidedRenderContext = false;
 			UsdUtils::FUsdPrimMaterialAssignmentInfo PrototypeAssignments = UsdUtils::GetPrimMaterialAssignments(
 				MeshPrim,
 				TimeCode,
 				bProvideMaterialIndices,
 				RenderContext,
 				UsdContext.CachedMeshConversionOptions.MaterialPurpose,
-				bForceCheckUnrealMaterialAttribute
+				bForceCheckUnrealMaterialAttribute,
+				bRequireMaterialsHaveProvidedRenderContext
 			);
 
 			if (Info.bInsideLODVariant)
@@ -759,13 +761,15 @@ namespace UE::Interchange::USD
 				const bool bProvideMaterialIndices = false;
 				const pxr::TfToken RenderContext = pxr::UsdShadeTokens->universalRenderContext;
 				const bool bForceCheckUnrealMaterialAttribute = true;
+				const bool bRequireMaterialsHaveProvidedRenderContext = false;
 				UsdUtils::FUsdPrimMaterialAssignmentInfo InstanceAssignments = UsdUtils::GetPrimMaterialAssignments(
 					Prim,
 					TimeCode,
 					bProvideMaterialIndices,
 					RenderContext,
 					UsdContext.CachedMeshConversionOptions.MaterialPurpose,
-					bForceCheckUnrealMaterialAttribute
+					bForceCheckUnrealMaterialAttribute,
+					bRequireMaterialsHaveProvidedRenderContext
 				);
 
 				SetSlotMaterialDependencies(SceneNode, InstanceAssignments, AccumulatedInfo, UsdContext);
diff --git a/Engine/Plugins/Interchange/Extensions/OpenUSD/Source/Import/Private/SchemaHandlers/VolumeSchemaHandler.cpp b/Engine/Plugins/Interchange/Extensions/OpenUSD/Source/Import/Private/SchemaHandlers/VolumeSchemaHandler.cpp
--- a/Engine/Plugins/Interchange/Extensions/OpenUSD/Source/Import/Private/SchemaHandlers/VolumeSchemaHandler.cpp
+++ b/Engine/Plugins/Interchange/Extensions/OpenUSD/Source/Import/Private/SchemaHandlers/VolumeSchemaHandler.cpp
@@ -190,14 +190,20 @@ namespace UE::VolVolumeSchemaHandler::Private
 
 		// Priority 1: Explicit material assignment on the Volume prim (Unreal materials).
 		{
+			// Don't pay too much attention to the render contexts right now, because we have no idea what our
+			// material schema handlers will do with this material
 			double TimeCode = UsdUtils::GetDefaultTimeCode();	 // Not relevant as material bindings can't be time sampled
 			bool bProvideMaterialIndices = false;
+			const bool bForceCheckUnrealMaterialAttribute = true;
+			const bool bRequireMaterialsHaveProvidedRenderContext = false;
 			UsdUtils::FUsdPrimMaterialAssignmentInfo Assignments = UsdUtils::GetPrimMaterialAssignments(
 				VolumePrim,
 				TimeCode,
 				bProvideMaterialIndices,
 				UnrealIdentifiers::Unreal,	  // Unreal materials should always be on the unreal render context
-				ConversionOptions.MaterialPurpose
+				ConversionOptions.MaterialPurpose,
+				bForceCheckUnrealMaterialAttribute,
+				bRequireMaterialsHaveProvidedRenderContext
 			);
 
 			for (const UsdUtils::FUsdPrimMaterialSlot& Slot : Assignments.Slots)
@@ -245,14 +251,20 @@ namespace UE::VolVolumeSchemaHandler::Private
 
 	UInterchangeShaderGraphNode* TryCreatingVolumetricShaderGraphNode(const UE::FUsdPrim& VolumePrim, UInterchangeUsdContext& UsdContext)
 	{
+		// Don't pay too much attention to the render contexts right now, because we have no idea what our
+		// material schema handlers will do with this material
 		const bool bProvideMaterialIndices = false;
 		const pxr::TfToken MtlxToken{"mtlx"};	 // We no longer need to pass this, but there's no reason *not* to pass mtlx here either
+		const bool bForceCheckUnrealMaterialAttribute = false;
+		const bool bRequireMaterialsHaveProvidedRenderContext = false;
 		UsdUtils::FUsdPrimMaterialAssignmentInfo Assignments = UsdUtils::GetPrimMaterialAssignments(
 			VolumePrim,
 			UsdUtils::GetDefaultTimeCode(),
 			bProvideMaterialIndices,
 			MtlxToken,
-			UsdContext.CachedMeshConversionOptions.MaterialPurpose
+			UsdContext.CachedMeshConversionOptions.MaterialPurpose,
+			bForceCheckUnrealMaterialAttribute,
+			bRequireMaterialsHaveProvidedRenderContext
 		);
 
 		for (const UsdUtils::FUsdPrimMaterialSlot& Slot : Assignments.Slots)
diff --git a/Engine/Plugins/Runtime/USDCore/Source/USDUtilities/Private/USDGeomMeshConversion.cpp b/Engine/Plugins/Runtime/USDCore/Source/USDUtilities/Private/USDGeomMeshConversion.cpp
--- a/Engine/Plugins/Runtime/USDCore/Source/USDUtilities/Private/USDGeomMeshConversion.cpp
+++ b/Engine/Plugins/Runtime/USDCore/Source/USDUtilities/Private/USDGeomMeshConversion.cpp
@@ -4213,7 +4213,8 @@ UsdUtils::FUsdPrimMaterialAssignmentInfo UsdUtils::GetPrimMaterialAssignments(
 	bool bProvideMaterialIndices,
 	const pxr::TfToken& RenderContext,
 	const pxr::TfToken& MaterialPurpose,
-	bool bForceCheckUnrealMaterialAttribute
+	bool bForceCheckUnrealMaterialAttribute,
+	bool bRequireMaterialsHaveProvidedRenderContext
 )
 {
 	TRACE_CPUPROFILER_EVENT_SCOPE(UsdUtils::GetPrimMaterialAssignments);
@@ -4262,27 +4263,38 @@ UsdUtils::FUsdPrimMaterialAssignmentInfo UsdUtils::GetPrimMaterialAssignments(
 		return {};
 	};
 
-	auto FetchMaterialByComputingBoundMaterial = [&MaterialPurpose](const pxr::UsdPrim& UsdPrim) -> TOptional<FString>
-	{
-		pxr::UsdShadeMaterialBindingAPI BindingAPI(UsdPrim);
-		pxr::UsdShadeMaterial ShadeMaterial = BindingAPI.ComputeBoundMaterial(MaterialPurpose);
-		if (!ShadeMaterial)
+	auto FetchMaterialByComputingBoundMaterial =
+		[&MaterialPurpose, &RenderContext, bRequireMaterialsHaveProvidedRenderContext]
+		(const pxr::UsdPrim& UsdPrim) -> TOptional<FString>
 		{
-			return {};
-		}
+			pxr::UsdShadeMaterialBindingAPI BindingAPI(UsdPrim);
+			pxr::UsdShadeMaterial ShadeMaterial = BindingAPI.ComputeBoundMaterial(MaterialPurpose);
+			if (!ShadeMaterial)
+			{
+				return {};
+			}
 
-		pxr::UsdPrim ShadeMaterialPrim = ShadeMaterial.GetPrim();
-		if (ShadeMaterialPrim)
-		{
-			const std::string ShadingEngineName = ShadeMaterialPrim.GetPrimPath().GetString();
-			if (!ShadingEngineName.empty())
+			if (bRequireMaterialsHaveProvidedRenderContext)
 			{
-				return UsdToUnreal::ConvertString(ShadingEngineName);
+				pxr::UsdShadeShader SurfaceShader = ShadeMaterial.ComputeSurfaceSource(RenderContext);
+				if (!SurfaceShader)
+				{
+					return {};
+				}
 			}
-		}
 
-		return {};
-	};
+			pxr::UsdPrim ShadeMaterialPrim = ShadeMaterial.GetPrim();
+			if (ShadeMaterialPrim)
+			{
+				const std::string ShadingEngineName = ShadeMaterialPrim.GetPrimPath().GetString();
+				if (!ShadingEngineName.empty())
+				{
+					return UsdToUnreal::ConvertString(ShadingEngineName);
+				}
+			}
+
+			return {};
+		};
 
 	FUsdPrimMaterialAssignmentInfo Result;
 
diff --git a/Engine/Plugins/Runtime/USDCore/Source/USDUtilities/Public/USDGeomMeshConversion.h b/Engine/Plugins/Runtime/USDCore/Source/USDUtilities/Public/USDGeomMeshConversion.h
--- a/Engine/Plugins/Runtime/USDCore/Source/USDUtilities/Public/USDGeomMeshConversion.h
+++ b/Engine/Plugins/Runtime/USDCore/Source/USDUtilities/Public/USDGeomMeshConversion.h
@@ -414,6 +414,8 @@ namespace UsdUtils
 	 * @param RenderContext - Which render context to get the materials for. Defaults to universal.
 	 * @param MaterialPurpose - Which material purpose to use when retrieving material bindings
 	 * @param bForceCheckUnrealMaterialAttribute - Whether to always check the "unrealMaterial" string attribute regardless of render context
+	 * @param bRequireMaterialsHaveProvidedRenderContext - If this is enabled, we'll only consider a material binding if it has a surface output
+	 * for the provided render context. If this is false, we'll use the bound material regardless
 	 * @return Struct containing an array of material assignments, and a corresponding array of material indices for all polygons of the prim,
 	 * matching the ordering of the material assignments.
 	 */
@@ -423,7 +425,8 @@ namespace UsdUtils
 		bool bProvideMaterialIndices = true,
 		const pxr::TfToken& RenderContext = pxr::UsdShadeTokens->universalRenderContext,
 		const pxr::TfToken& MaterialPurpose = pxr::UsdShadeTokens->allPurpose,
-		bool bForceCheckUnrealMaterialAttribute = false
+		bool bForceCheckUnrealMaterialAttribute = false,
+		bool bRequireMaterialsHaveProvidedRenderContext = true
 	);
 
 	/** Returns whether this prim can be interpreted as describing a static mesh with multiple LODs (i.e. if the prim holds the LOD variant set
