USD Interchange: Expose substrate parent materials like the other UsdPreviewSurface parent materials.

#Synced-CL 51466015

diff --git a/Engine/Plugins/Interchange/Extensions/OpenUSD/Source/Import/Private/InterchangeUSDPipeline.cpp b/Engine/Plugins/Interchange/Extensions/OpenUSD/Source/Import/Private/InterchangeUSDPipeline.cpp
--- a/Engine/Plugins/Interchange/Extensions/OpenUSD/Source/Import/Private/InterchangeUSDPipeline.cpp
+++ b/Engine/Plugins/Interchange/Extensions/OpenUSD/Source/Import/Private/InterchangeUSDPipeline.cpp
@@ -1410,10 +1410,23 @@ void UInterchangeUsdPipeline::ExecutePipeline(
 		return;
 	}
 
-	if (IInterchangeImportModule::Get().IsSubstrateEnabled())
+	const UUsdProjectSettings* UsdProjectSettings = GetDefault<UUsdProjectSettings>();
+	if (IInterchangeImportModule::Get().IsSubstrateEnabled() && UsdProjectSettings)
 	{
+		// Let's keep Vertex Color parent material for the moment (maybe we could instead set UseVertexColor in the parent material for Substrate?)
+		TMap<FString, FString> OldToNewMaterial = {
+			{UsdProjectSettings->ReferencePreviewSurfaceMaterial.ToString(), 						UsdProjectSettings->ReferenceSubstratePreviewSurfaceMaterial.ToString()},
+			{UsdProjectSettings->ReferencePreviewSurfaceTranslucentMaterial.ToString(), 			UsdProjectSettings->ReferenceSubstratePreviewSurfaceMaterial.ToString()},
+			{UsdProjectSettings->ReferencePreviewSurfaceTwoSidedMaterial.ToString(), 				UsdProjectSettings->ReferenceSubstratePreviewSurfaceMaterial.ToString()},
+			{UsdProjectSettings->ReferencePreviewSurfaceTranslucentTwoSidedMaterial.ToString(), 	UsdProjectSettings->ReferenceSubstratePreviewSurfaceMaterial.ToString()},
+			{UsdProjectSettings->ReferencePreviewSurfaceVTMaterial.ToString(), 						UsdProjectSettings->ReferenceSubstratePreviewSurfaceVTMaterial.ToString()},
+			{UsdProjectSettings->ReferencePreviewSurfaceTwoSidedVTMaterial.ToString(), 				UsdProjectSettings->ReferenceSubstratePreviewSurfaceVTMaterial.ToString()},
+			{UsdProjectSettings->ReferencePreviewSurfaceTranslucentVTMaterial.ToString(), 			UsdProjectSettings->ReferenceSubstratePreviewSurfaceVTMaterial.ToString()},
+			{UsdProjectSettings->ReferencePreviewSurfaceTranslucentTwoSidedVTMaterial.ToString(), 	UsdProjectSettings->ReferenceSubstratePreviewSurfaceVTMaterial.ToString()},
+		};
+
 		NodeContainer->IterateNodesOfType<UInterchangeMaterialInstanceNode>(
-			[NodeContainer](const FString& NodeUID, UInterchangeMaterialInstanceNode* MaterialInstanceNode)
+			[NodeContainer, &OldToNewMaterial](const FString& NodeUID, UInterchangeMaterialInstanceNode* MaterialInstanceNode)
 			{
 				// We can have a parent material that comes from a ShaderGraphNode (for example a TwoSided instance created based on the ShaderGraph),
 				// in that case it's not a UsdPreviewSurface
@@ -1425,9 +1438,13 @@ void UInterchangeUsdPipeline::ExecutePipeline(
 				}
 
 				FString ParentPath;
-				MaterialInstanceNode->GetCustomParent(ParentPath);
-				// Let's keep Vertex Color parent material for the moment (maybe we could instead set UseVertexColor in the parent material for Substrate?)
-				if (ParentPath.Contains("DisplayColor", ESearchCase::CaseSensitive))
+				if (!MaterialInstanceNode->GetCustomParent(ParentPath))
+				{
+					return;
+				}
+
+				FString* NewParentPath = OldToNewMaterial.Find(ParentPath);
+				if (!NewParentPath)
 				{
 					return;
 				}
@@ -1435,15 +1452,11 @@ void UInterchangeUsdPipeline::ExecutePipeline(
 				const FString FactoryNodeUID = UInterchangeFactoryBaseNode::BuildFactoryNodeUid(NodeUID);
 				if (UInterchangeMaterialInstanceFactoryNode* MaterialInstanceFactoryNode = Cast<UInterchangeMaterialInstanceFactoryNode>(NodeContainer->GetFactoryNode(FactoryNodeUID)))
 				{
-					MaterialInstanceFactoryNode->SetCustomParent(TEXT("/InterchangeAssets/Substrate/UsdPreviewSurface.UsdPreviewSurface"));
+					MaterialInstanceFactoryNode->SetCustomParent(*NewParentPath);
 
 					if (int32 Value; MaterialInstanceNode->GetInt32Attribute(UE::Interchange::USD::UsdReferenceMaterialProperties, Value))
 					{
 						EUsdReferenceMaterialProperties Properties = static_cast<EUsdReferenceMaterialProperties>(Value);
-						if (EnumHasAnyFlags(Properties, EUsdReferenceMaterialProperties::VT))
-						{
-							MaterialInstanceFactoryNode->SetCustomParent(TEXT("/InterchangeAssets/Substrate/UsdPreviewSurfaceVT.UsdPreviewSurfaceVT"));
-						}
 
 						if (EnumHasAnyFlags(Properties, EUsdReferenceMaterialProperties::Translucent))
 						{
diff --git a/Engine/Plugins/Runtime/USDCore/Source/USDClasses/Public/USDProjectSettings.h b/Engine/Plugins/Runtime/USDCore/Source/USDClasses/Public/USDProjectSettings.h
--- a/Engine/Plugins/Runtime/USDCore/Source/USDClasses/Public/USDProjectSettings.h
+++ b/Engine/Plugins/Runtime/USDCore/Source/USDClasses/Public/USDProjectSettings.h
@@ -44,6 +44,10 @@ namespace UsdDefaultReferenceMaterialPaths
 	const inline FSoftObjectPath ReferencePreviewSurfaceTranslucentVTMaterial{TEXT("/USDCore/Materials/UsdPreviewSurfaceTranslucentVT.UsdPreviewSurfaceTranslucentVT")};
 	const inline FSoftObjectPath ReferencePreviewSurfaceTwoSidedVTMaterial{TEXT("/USDCore/Materials/UsdPreviewSurfaceTwoSidedVT.UsdPreviewSurfaceTwoSidedVT")};
 	const inline FSoftObjectPath ReferencePreviewSurfaceTranslucentTwoSidedVTMaterial{TEXT("/USDCore/Materials/UsdPreviewSurfaceTranslucentTwoSidedVT.UsdPreviewSurfaceTranslucentTwoSidedVT")};
+
+	const inline FSoftObjectPath ReferenceSubstratePreviewSurfaceMaterial{TEXT("/InterchangeAssets/Substrate/UsdPreviewSurface.UsdPreviewSurface")};
+	const inline FSoftObjectPath ReferenceSubstratePreviewSurfaceVTMaterial{TEXT("/InterchangeAssets/Substrate/UsdPreviewSurfaceVT.UsdPreviewSurfaceVT")};
+
 	const inline FSoftObjectPath ReferenceDisplayColorMaterial{TEXT("/USDCore/Materials/DisplayColor.DisplayColor")};
 	const inline FSoftObjectPath ReferenceDisplayColorAndOpacityMaterial{TEXT("/USDCore/Materials/DisplayColorAndOpacity.DisplayColorAndOpacity")};
 	const inline FSoftObjectPath ReferenceDisplayColorTwoSidedMaterial{TEXT("/USDCore/Materials/DisplayColorTwoSided.DisplayColorTwoSided")};
@@ -258,6 +262,26 @@ public:
 	UPROPERTY(config, EditAnywhere, Category = "USD|Reference Materials|UsdPreviewSurface with Virtual Textures", meta = (AllowedClasses = "/Script/Engine.MaterialInterface"))
 	FSoftObjectPath ReferencePreviewSurfaceTranslucentTwoSidedVTMaterial = UsdDefaultReferenceMaterialPaths::ReferencePreviewSurfaceTranslucentTwoSidedVTMaterial;
 
+	/**
+	 * What material to use as reference material when creating material instances from USD materials.
+	 * Exclusively used by USD imports via the Interchange framework, when Substrate is enabled for the project.
+	 * You can swap these with your own materials, but make sure that the replacement materials have parameters with
+	 * the same names and types as the ones provided by the default material, otherwise the instances will not have
+	 * the parameters filled with values extracted from the USD material when parsing.
+	 */
+	UPROPERTY(config, EditAnywhere, Category = "USD|Reference Materials|UsdPreviewSurface", meta = (AllowedClasses = "/Script/Engine.MaterialInterface"))
+	FSoftObjectPath ReferenceSubstratePreviewSurfaceMaterial = UsdDefaultReferenceMaterialPaths::ReferenceSubstratePreviewSurfaceMaterial;
+
+	/**
+	 * What material to use as reference material when creating material instances from USD materials.
+	 * Exclusively used by USD imports via the Interchange framework, when Substrate is enabled for the project.
+	 * You can swap these with your own materials, but make sure that the replacement materials have parameters with
+	 * the same names and types as the ones provided by the default material, otherwise the instances will not have
+	 * the parameters filled with values extracted from the USD material when parsing.
+	 */
+	UPROPERTY(config, EditAnywhere, Category = "USD|Reference Materials|UsdPreviewSurface", meta = (AllowedClasses = "/Script/Engine.MaterialInterface"))
+	FSoftObjectPath ReferenceSubstratePreviewSurfaceVTMaterial = UsdDefaultReferenceMaterialPaths::ReferenceSubstratePreviewSurfaceVTMaterial;
+
 	/**
 	 * What material to use as reference material when creating material instances from USD materials.
 	 * You can swap these with your own materials, but make sure that the replacement materials have parameters with
