USD: Fix schema name when trying to find Mesh prims.

#Synced-CL 51466015

diff --git a/Engine/Plugins/Interchange/Extensions/OpenUSD/Source/Import/Private/InterchangeUSDInfoCache.cpp b/Engine/Plugins/Interchange/Extensions/OpenUSD/Source/Import/Private/InterchangeUSDInfoCache.cpp
--- a/Engine/Plugins/Interchange/Extensions/OpenUSD/Source/Import/Private/InterchangeUSDInfoCache.cpp
+++ b/Engine/Plugins/Interchange/Extensions/OpenUSD/Source/Import/Private/InterchangeUSDInfoCache.cpp
@@ -761,7 +761,7 @@ namespace UE::UsdInfoCache::Private
 		// If the stage doesn't contain any animated mesh prims, then don't bother doing a full check
 		bool bHasAnimatedMesh = false;
 		{
-			for (const UE::FUsdPrim& ChildPrim : UsdUtils::GetAllPrimsOfType(PseudoRoot, TEXT("Mesh")))
+			for (const UE::FUsdPrim& ChildPrim : UsdUtils::GetAllPrimsOfType(PseudoRoot, TEXT("UsdGeomMesh")))
 			{
 				if (UsdUtils::IsAnimatedMesh(ChildPrim))
 				{
