USD: Fix MaterialX material assignments being parsed via USD Interchange with an extra world grid material slot

#Synced-CL 51466015

diff --git a/Engine/Plugins/Interchange/Extensions/OpenUSD/Source/Import/Private/InterchangeUsdContext.cpp b/Engine/Plugins/Interchange/Extensions/OpenUSD/Source/Import/Private/InterchangeUsdContext.cpp
--- a/Engine/Plugins/Interchange/Extensions/OpenUSD/Source/Import/Private/InterchangeUsdContext.cpp
+++ b/Engine/Plugins/Interchange/Extensions/OpenUSD/Source/Import/Private/InterchangeUsdContext.cpp
@@ -279,6 +279,12 @@ void UInterchangeUsdContext::SetupMeshConversionOptions()
 		CachedMeshConversionOptions.MaterialPurpose = Settings->MaterialPurpose.IsNone()
 														  ? pxr::UsdShadeTokens->allPurpose
 														  : pxr::TfToken{TCHAR_TO_UTF8(*Settings->MaterialPurpose.ToString())};
+
+		// Interchange doesn't know the final render context at the time we collect material assignments
+		// (we query with universalRenderContext and resolve render contexts at a later stage via schema handlers),
+		// so we accept any material binding regardless of whether it provides a surface shader for the render context.
+		CachedMeshConversionOptions.bForceCheckUnrealMaterialAttribute = true;
+		CachedMeshConversionOptions.bRequireMaterialsHaveProvidedRenderContext = false;
 	}
 #endif	  // USE_USD_SDK
 }
diff --git a/Engine/Plugins/Interchange/Extensions/OpenUSD/Source/Import/Private/SchemaHandlers/GprimSchemaHandler.cpp b/Engine/Plugins/Interchange/Extensions/OpenUSD/Source/Import/Private/SchemaHandlers/GprimSchemaHandler.cpp
--- a/Engine/Plugins/Interchange/Extensions/OpenUSD/Source/Import/Private/SchemaHandlers/GprimSchemaHandler.cpp
+++ b/Engine/Plugins/Interchange/Extensions/OpenUSD/Source/Import/Private/SchemaHandlers/GprimSchemaHandler.cpp
@@ -291,16 +291,14 @@ namespace UE::GeomGprimSchemaHandler::Private
 			double TimeCode = UsdUtils::GetDefaultTimeCode();
 			const bool bProvideMaterialIndices = false;
 			const pxr::TfToken RenderContext = pxr::UsdShadeTokens->universalRenderContext;
-			const bool bForceCheckUnrealMaterialAttribute = true;
-			const bool bRequireMaterialsHaveProvidedRenderContext = false;
 			UsdUtils::FUsdPrimMaterialAssignmentInfo PrototypeAssignments = UsdUtils::GetPrimMaterialAssignments(
 				MeshPrim,
 				TimeCode,
 				bProvideMaterialIndices,
 				RenderContext,
 				UsdContext.CachedMeshConversionOptions.MaterialPurpose,
-				bForceCheckUnrealMaterialAttribute,
-				bRequireMaterialsHaveProvidedRenderContext
+				UsdContext.CachedMeshConversionOptions.bForceCheckUnrealMaterialAttribute,
+				UsdContext.CachedMeshConversionOptions.bRequireMaterialsHaveProvidedRenderContext
 			);
 
 			if (Info.bInsideLODVariant)
@@ -760,16 +758,14 @@ namespace UE::Interchange::USD
 				const double TimeCode = UsdUtils::GetDefaultTimeCode();
 				const bool bProvideMaterialIndices = false;
 				const pxr::TfToken RenderContext = pxr::UsdShadeTokens->universalRenderContext;
-				const bool bForceCheckUnrealMaterialAttribute = true;
-				const bool bRequireMaterialsHaveProvidedRenderContext = false;
 				UsdUtils::FUsdPrimMaterialAssignmentInfo InstanceAssignments = UsdUtils::GetPrimMaterialAssignments(
 					Prim,
 					TimeCode,
 					bProvideMaterialIndices,
 					RenderContext,
 					UsdContext.CachedMeshConversionOptions.MaterialPurpose,
-					bForceCheckUnrealMaterialAttribute,
-					bRequireMaterialsHaveProvidedRenderContext
+					UsdContext.CachedMeshConversionOptions.bForceCheckUnrealMaterialAttribute,
+					UsdContext.CachedMeshConversionOptions.bRequireMaterialsHaveProvidedRenderContext
 				);
 
 				SetSlotMaterialDependencies(SceneNode, InstanceAssignments, AccumulatedInfo, UsdContext);
diff --git a/Engine/Plugins/Interchange/Extensions/OpenUSD/Source/Import/Private/SchemaHandlers/VolumeSchemaHandler.cpp b/Engine/Plugins/Interchange/Extensions/OpenUSD/Source/Import/Private/SchemaHandlers/VolumeSchemaHandler.cpp
--- a/Engine/Plugins/Interchange/Extensions/OpenUSD/Source/Import/Private/SchemaHandlers/VolumeSchemaHandler.cpp
+++ b/Engine/Plugins/Interchange/Extensions/OpenUSD/Source/Import/Private/SchemaHandlers/VolumeSchemaHandler.cpp
@@ -194,16 +194,14 @@ namespace UE::VolVolumeSchemaHandler::Private
 			// material schema handlers will do with this material
 			double TimeCode = UsdUtils::GetDefaultTimeCode();	 // Not relevant as material bindings can't be time sampled
 			bool bProvideMaterialIndices = false;
-			const bool bForceCheckUnrealMaterialAttribute = true;
-			const bool bRequireMaterialsHaveProvidedRenderContext = false;
 			UsdUtils::FUsdPrimMaterialAssignmentInfo Assignments = UsdUtils::GetPrimMaterialAssignments(
 				VolumePrim,
 				TimeCode,
 				bProvideMaterialIndices,
 				UnrealIdentifiers::Unreal,	  // Unreal materials should always be on the unreal render context
 				ConversionOptions.MaterialPurpose,
-				bForceCheckUnrealMaterialAttribute,
-				bRequireMaterialsHaveProvidedRenderContext
+				ConversionOptions.bForceCheckUnrealMaterialAttribute,
+				ConversionOptions.bRequireMaterialsHaveProvidedRenderContext
 			);
 
 			for (const UsdUtils::FUsdPrimMaterialSlot& Slot : Assignments.Slots)
@@ -256,7 +254,6 @@ namespace UE::VolVolumeSchemaHandler::Private
 		const bool bProvideMaterialIndices = false;
 		const pxr::TfToken MtlxToken{"mtlx"};	 // We no longer need to pass this, but there's no reason *not* to pass mtlx here either
 		const bool bForceCheckUnrealMaterialAttribute = false;
-		const bool bRequireMaterialsHaveProvidedRenderContext = false;
 		UsdUtils::FUsdPrimMaterialAssignmentInfo Assignments = UsdUtils::GetPrimMaterialAssignments(
 			VolumePrim,
 			UsdUtils::GetDefaultTimeCode(),
@@ -264,7 +261,7 @@ namespace UE::VolVolumeSchemaHandler::Private
 			MtlxToken,
 			UsdContext.CachedMeshConversionOptions.MaterialPurpose,
 			bForceCheckUnrealMaterialAttribute,
-			bRequireMaterialsHaveProvidedRenderContext
+			UsdContext.CachedMeshConversionOptions.bRequireMaterialsHaveProvidedRenderContext
 		);
 
 		for (const UsdUtils::FUsdPrimMaterialSlot& Slot : Assignments.Slots)
diff --git a/Engine/Plugins/Runtime/USDCore/Source/USDUtilities/Private/USDGeomMeshConversion.cpp b/Engine/Plugins/Runtime/USDCore/Source/USDUtilities/Private/USDGeomMeshConversion.cpp
--- a/Engine/Plugins/Runtime/USDCore/Source/USDUtilities/Private/USDGeomMeshConversion.cpp
+++ b/Engine/Plugins/Runtime/USDCore/Source/USDUtilities/Private/USDGeomMeshConversion.cpp
@@ -1833,7 +1833,9 @@ namespace UE::UsdGeomMeshConversion::Private
 				Options.TimeCode,
 				bProvideMaterialIndices,
 				Options.RenderContext,
-				Options.MaterialPurpose
+				Options.MaterialPurpose,
+				Options.bForceCheckUnrealMaterialAttribute,
+				Options.bRequireMaterialsHaveProvidedRenderContext
 			);
 
 			InOutMeshData.MaterialIndexOffset = InOutCombinedMaterialAssignments.Slots.Num();
@@ -3576,6 +3578,8 @@ namespace UsdToUnreal
 		, SubdivisionLevel(0)
 		, ImportPrimvars(EImportPrimvar::Standard)
 		, bOffsetJointIndicesForMultipleRootBones(true)
+		, bForceCheckUnrealMaterialAttribute(false)
+		, bRequireMaterialsHaveProvidedRenderContext(true)
 	{
 	}
 }
diff --git a/Engine/Plugins/Runtime/USDCore/Source/USDUtilities/Public/USDGeomMeshConversion.h b/Engine/Plugins/Runtime/USDCore/Source/USDUtilities/Public/USDGeomMeshConversion.h
--- a/Engine/Plugins/Runtime/USDCore/Source/USDUtilities/Public/USDGeomMeshConversion.h
+++ b/Engine/Plugins/Runtime/USDCore/Source/USDUtilities/Public/USDGeomMeshConversion.h
@@ -115,6 +115,17 @@ namespace UsdToUnreal
 
 		// Whether to add an additional offset to vertex join indices if the skinned mesh's skeleton has multiple root bones
 		bool bOffsetJointIndicesForMultipleRootBones;
+
+		// Whether to always check the deprecated "unrealMaterial" string attribute on prims regardless of the configured RenderContext.
+		// When false (the default), the unrealMaterial attribute is only checked when the RenderContext is "unreal".
+		// This is mainly used by Interchange: The render context is resolved only later (and we always provide 'universal'), so we need
+		// this additional parameter to not ignore the string attribute completely
+		bool bForceCheckUnrealMaterialAttribute;
+
+		// If true, a bound material is only used if it has a surface shader output for the configured RenderContext.
+		// If it doesn't, the binding is rejected and processing falls through to the DisplayColor fallback.
+		// Legacy USD importers set this to true; Interchange sets it to false because the render context is resolved at a later stage.
+		bool bRequireMaterialsHaveProvidedRenderContext;
 	};
 
 	/**
