USD: Partial revert of initial approach for handling inputs:normalize.

#Synced-CL 52058314

diff --git a/Engine/Plugins/Interchange/Extensions/OpenUSD/Source/Import/Private/SchemaHandlers/LightSchemaHandler.cpp b/Engine/Plugins/Interchange/Extensions/OpenUSD/Source/Import/Private/SchemaHandlers/LightSchemaHandler.cpp
--- a/Engine/Plugins/Interchange/Extensions/OpenUSD/Source/Import/Private/SchemaHandlers/LightSchemaHandler.cpp
+++ b/Engine/Plugins/Interchange/Extensions/OpenUSD/Source/Import/Private/SchemaHandlers/LightSchemaHandler.cpp
@@ -121,6 +121,11 @@ namespace UE::LuxLightSchemaHandler::Private
 		bool bNormalizeAuthored = false;
 		bool bNormalizedIntensityUsd = UsdUtils::GetLightAttrValueWithInputsFallback<bool>(Prim.GetAttribute(*NormalizeToken), UsdUtils::GetDefaultTimeCode(), &bNormalizeAuthored);
 		bNormalizedIntensityUsd = bNormalizeAuthored && bNormalizedIntensityUsd;
+		if (bNormalizeAuthored)
+		{
+			USD_LOG_WARNING(TEXT("Light prim %s has 'inputs:normalize' authored, which is not supported!"), *Prim.GetPrimPath().GetString());
+			LightNode->AddBooleanAttribute(LightAPINormalizeAttributeKey, bNormalizedIntensityUsd);
+		}
 
 		static const FString RadiusToken = UsdToUnreal::ConvertToken(pxr::UsdLuxTokens->inputsRadius);
 
@@ -194,7 +199,6 @@ namespace UE::LuxLightSchemaHandler::Private
 
 		float SolidAngle = 1.0f;
 		float AreaSqMeters = 1.0f;
-		bool bAllowNormalizingIntensity = true;
 		if (Prim.IsA(TEXT("DistantLight")))
 		{
 			UInterchangeDirectionalLightNode* CastLightNode = Cast<UInterchangeDirectionalLightNode>(LightNode);
@@ -203,10 +207,6 @@ namespace UE::LuxLightSchemaHandler::Private
 				return;
 			}
 
-			// For distant lights we always use Lux, for which normalization by area doesn't make much sense (it's already
-			// a per-area sort of unit)
-			bAllowNormalizingIntensity = false;
-
 			static const FString SourceAngleToken = UsdToUnreal::ConvertToken(pxr::UsdLuxTokens->inputsAngle);
 			bool bAngleAuthored = false;
 			float UsdAngle = UsdUtils::GetLightAttrValueWithInputsFallback<float>(Prim.GetAttribute(*SourceAngleToken), UsdUtils::GetDefaultTimeCode(), &bAngleAuthored); 
@@ -330,10 +330,6 @@ namespace UE::LuxLightSchemaHandler::Private
 				return;
 			}
 
-			// For dome lights the intensity is just a simple multiplier with no physical meaning, so normalization
-			// doesn't make much sense
-			bAllowNormalizingIntensity = false;
-
 			// Show some warning in case the file specifies a different DomeLight format, as we only support latlong
 			static const FString FormatToken = UsdToUnreal::ConvertToken(pxr::UsdLuxTokens->inputsTextureFormat);
 			UE::FUsdAttribute FormatAttribute = Prim.GetAttribute(*FormatToken);
@@ -407,33 +403,14 @@ namespace UE::LuxLightSchemaHandler::Private
 		{
 			// This will end up being Nits for local lights (point, rect, disk, spot), lux for distant lights,
 			// and a simple multiplier without units for dome lights
-			float Intensity = UsdToUnreal::ConvertLightIntensity(IntensityUsd, Exposure);
-
-			// We only normalize the intensity for local lights here as we can flip them into Lumen to represent
-			// normalized units. For distant/dome lights we just ignore normalization
-			//
-			// It's a bit awkward because USD does define normalization for distant lights at least, but the normalization
-			// USD suggests involves mainly hard-coding a sizeFactor constant (which we won't do as there's no good place
-			// to put it)
-			float IntensityUE = 1.0;
-			EInterchangeLightUnits UnitsUE = EInterchangeLightUnits::Unitless;
-			if (bNormalizedIntensityUsd && bAllowNormalizingIntensity)
-			{
-				IntensityUE = UsdUtils::ConvertIntensityFromNits(Intensity, SolidAngle, AreaSqMeters, ELightUnits::Lumens);
-				UnitsUE = EInterchangeLightUnits::Lumens;
-			}
-			else
-			{
-				IntensityUE = Intensity;
-				UnitsUE = EInterchangeLightUnits::Nits;
-			}
+			const float Intensity = UsdToUnreal::ConvertLightIntensity(IntensityUsd, Exposure);
 
-			LightNode->SetCustomIntensity(IntensityUE);
+			LightNode->SetCustomIntensity(Intensity);
 			if (UInterchangeLightNode* DerivedLightNode = Cast<UInterchangeLightNode>(LightNode))
 			{
 				// Note that directional and sky lights are not LightNodes (they derive BaseLightNode directly), so this
 				// check prevents us from trying to set them with Nits
-				DerivedLightNode->SetCustomIntensityUnits(UnitsUE);
+				DerivedLightNode->SetCustomIntensityUnits(EInterchangeLightUnits::Nits);
 			}
 		}
 	}
@@ -529,10 +506,8 @@ namespace UE::Interchange::USD
 				const static TMap<FString, TArray<FInterchangeTrackInfo>> CombinedLightMapping = []()
 				{
 					TMap<FString, TArray<FInterchangeTrackInfo>> Combined = CommonAttributeMapping;
-					Combined.Add(UsdToUnreal::ConvertToken(pxr::UsdLuxTokens->inputsShapingConeAngle), 		{{UnrealIdentifiers::OuterConeAnglePropertyName, 	EInterchangePropertyTracks::LightOuterConeAngle},
-																											 {UnrealIdentifiers::IntensityPropertyName, 		EInterchangePropertyTracks::LightIntensity}});
-					Combined.Add(UsdToUnreal::ConvertToken(pxr::UsdLuxTokens->inputsShapingConeSoftness), 	{{UnrealIdentifiers::InnerConeAnglePropertyName, 	EInterchangePropertyTracks::LightInnerConeAngle},
-																											 {UnrealIdentifiers::IntensityPropertyName, 		EInterchangePropertyTracks::LightIntensity}});
+					Combined.Add(UsdToUnreal::ConvertToken(pxr::UsdLuxTokens->inputsShapingConeAngle), 		{{UnrealIdentifiers::OuterConeAnglePropertyName, 	EInterchangePropertyTracks::LightOuterConeAngle}});
+					Combined.Add(UsdToUnreal::ConvertToken(pxr::UsdLuxTokens->inputsShapingConeSoftness), 	{{UnrealIdentifiers::InnerConeAnglePropertyName, 	EInterchangePropertyTracks::LightInnerConeAngle}});
 					return Combined;
 				}();
 				AttributeMapping = &CombinedLightMapping;
@@ -542,8 +517,7 @@ namespace UE::Interchange::USD
 				const static TMap<FString, TArray<FInterchangeTrackInfo>> CombinedLightMapping = []()
 				{
 					TMap<FString, TArray<FInterchangeTrackInfo>> Combined = CommonAttributeMapping;
-					Combined.Add(UsdToUnreal::ConvertToken(pxr::UsdLuxTokens->inputsRadius), 	{{UnrealIdentifiers::SourceRadiusPropertyName, 		EInterchangePropertyTracks::LightSourceRadius},
-																								 {UnrealIdentifiers::IntensityPropertyName, 		EInterchangePropertyTracks::LightIntensity}});
+					Combined.Add(UsdToUnreal::ConvertToken(pxr::UsdLuxTokens->inputsRadius), 	{{UnrealIdentifiers::SourceRadiusPropertyName, 		EInterchangePropertyTracks::LightSourceRadius}});
 					return Combined;
 				}();
 				AttributeMapping = &CombinedLightMapping;
@@ -555,12 +529,9 @@ namespace UE::Interchange::USD
 					TMap<FString, TArray<FInterchangeTrackInfo>> Combined = CommonAttributeMapping;
 					// Animated inputsRadius on a RectLightNode means this is a DiskLight handled as a RectLight
 					Combined.Add(UsdToUnreal::ConvertToken(pxr::UsdLuxTokens->inputsRadius), 	{{UnrealIdentifiers::SourceWidthPropertyName, 		EInterchangePropertyTracks::LightSourceWidth},
-																								 {UnrealIdentifiers::SourceHeightPropertyName, 		EInterchangePropertyTracks::LightSourceHeight},
-																								 {UnrealIdentifiers::IntensityPropertyName, 		EInterchangePropertyTracks::LightIntensity}});
-					Combined.Add(UsdToUnreal::ConvertToken(pxr::UsdLuxTokens->inputsWidth), 	{{UnrealIdentifiers::SourceWidthPropertyName, 		EInterchangePropertyTracks::LightSourceWidth},
-																								 {UnrealIdentifiers::IntensityPropertyName, 		EInterchangePropertyTracks::LightIntensity}});
-					Combined.Add(UsdToUnreal::ConvertToken(pxr::UsdLuxTokens->inputsHeight), 	{{UnrealIdentifiers::SourceHeightPropertyName, 		EInterchangePropertyTracks::LightSourceHeight},
-																								 {UnrealIdentifiers::IntensityPropertyName, 		EInterchangePropertyTracks::LightIntensity}});
+																								 {UnrealIdentifiers::SourceHeightPropertyName, 		EInterchangePropertyTracks::LightSourceHeight}});
+					Combined.Add(UsdToUnreal::ConvertToken(pxr::UsdLuxTokens->inputsWidth), 	{{UnrealIdentifiers::SourceWidthPropertyName, 		EInterchangePropertyTracks::LightSourceWidth}});
+					Combined.Add(UsdToUnreal::ConvertToken(pxr::UsdLuxTokens->inputsHeight), 	{{UnrealIdentifiers::SourceHeightPropertyName, 		EInterchangePropertyTracks::LightSourceHeight}});
 					return Combined;
 				}();
 				AttributeMapping = &CombinedLightMapping;
diff --git a/Engine/Plugins/Interchange/Runtime/Source/Common/Public/Usd/InterchangeUsdDefinitions.h b/Engine/Plugins/Interchange/Runtime/Source/Common/Public/Usd/InterchangeUsdDefinitions.h
--- a/Engine/Plugins/Interchange/Runtime/Source/Common/Public/Usd/InterchangeUsdDefinitions.h
+++ b/Engine/Plugins/Interchange/Runtime/Source/Common/Public/Usd/InterchangeUsdDefinitions.h
@@ -66,6 +66,10 @@ namespace UE::Interchange::USD
 	// the actual subdivision
 	inline const FString SubdivisionLevelAttributeKey = TEXT("USD_SubdivisionLevel");
 
+	// We don't fully support inputs:normalize for USD lights, but we at least move the attribute into the Interchange
+	// factory node so that it's exposed for potential user pipelines
+	inline const FString LightAPINormalizeAttributeKey = TEXT("USD_LightAPIInputsNormalize");
+
 	// Name of the variant set that is used for describing LODs
 	inline const FString LODString = TEXT("LOD");
 
diff --git a/Engine/Plugins/Runtime/USDCore/Source/USDUtilities/Private/USDLightConversion.cpp b/Engine/Plugins/Runtime/USDCore/Source/USDUtilities/Private/USDLightConversion.cpp
--- a/Engine/Plugins/Runtime/USDCore/Source/USDUtilities/Private/USDLightConversion.cpp
+++ b/Engine/Plugins/Runtime/USDCore/Source/USDUtilities/Private/USDLightConversion.cpp
@@ -173,24 +173,13 @@ bool UsdToUnreal::ConvertRectLight(const pxr::UsdPrim& Prim, URectLightComponent
 	const float UsdExposure = UsdUtils::GetLightAttrValueWithInputsFallback<float>(RectLight.GetExposureAttr(), UsdTimeCode);
 	const float UsdWidth = UsdUtils::GetLightAttrValueWithInputsFallback<float>(RectLight.GetWidthAttr(), UsdTimeCode);
 	const float UsdHeight = UsdUtils::GetLightAttrValueWithInputsFallback<float>(RectLight.GetHeightAttr(), UsdTimeCode);
-	const double UsdNormalized = UsdUtils::GetLightAttrValueWithInputsFallback<bool>(RectLight.GetNormalizeAttr(), pxr::UsdTimeCode::Default().GetValue());
 
 	const double NitsIntensity = UsdToUnreal::ConvertLightIntensity(UsdIntensity, UsdExposure);
 
 	LightComponent.SourceWidth = UsdToUnreal::ConvertDistance(StageInfo, UsdWidth);
 	LightComponent.SourceHeight = UsdToUnreal::ConvertDistance(StageInfo, UsdHeight);
-	if (UsdNormalized)
-	{
-		const float SolidAngle = UsdUtils::GetRectLightSolidAngle();
-		const float AreaSqMeters = UsdUtils::GetRectLightArea(LightComponent.SourceWidth / 100.0f, LightComponent.SourceHeight / 100.0f);
-		LightComponent.IntensityUnits = ELightUnits::Lumens;
-		LightComponent.Intensity = UsdUtils::ConvertIntensityFromNits(NitsIntensity, SolidAngle, AreaSqMeters, ELightUnits::Lumens);
-	}
-	else
-	{
-		LightComponent.Intensity = NitsIntensity;
-		LightComponent.IntensityUnits = ELightUnits::Nits;
-	}
+	LightComponent.Intensity = NitsIntensity;
+	LightComponent.IntensityUnits = ELightUnits::Nits;
 
 	return true;
 }
@@ -210,7 +199,6 @@ bool UsdToUnreal::ConvertDiskLight(const pxr::UsdPrim& Prim, URectLightComponent
 	const float UsdIntensity = UsdUtils::GetLightAttrValueWithInputsFallback<float>(DiskLight.GetIntensityAttr(), UsdTimeCode);
 	const float UsdExposure = UsdUtils::GetLightAttrValueWithInputsFallback<float>(DiskLight.GetExposureAttr(), UsdTimeCode);
 	const float UsdRadius = UsdUtils::GetLightAttrValueWithInputsFallback<float>(DiskLight.GetRadiusAttr(), UsdTimeCode);
-	const double UsdNormalized = UsdUtils::GetLightAttrValueWithInputsFallback<bool>(DiskLight.GetNormalizeAttr(), pxr::UsdTimeCode::Default().GetValue());
 
 	const float RadiusUE = UsdToUnreal::ConvertDistance(StageInfo, UsdRadius);
 	const float SquareSideUE = UsdUtils::GetRectLightSideFromDiskLight(RadiusUE);
@@ -218,19 +206,8 @@ bool UsdToUnreal::ConvertDiskLight(const pxr::UsdPrim& Prim, URectLightComponent
 
 	LightComponent.SourceWidth = SquareSideUE;
 	LightComponent.SourceHeight = SquareSideUE;
-	if (UsdNormalized)
-	{
-		const float RadiusMeters = RadiusUE / 100.0f;
-		const float SolidAngle = UsdUtils::GetDiskLightSolidAngle();
-		const float AreaSqMeters = UsdUtils::GetDiskLightArea(RadiusMeters);
-		LightComponent.IntensityUnits = ELightUnits::Lumens;
-		LightComponent.Intensity = UsdUtils::ConvertIntensityFromNits(NitsIntensity, SolidAngle, AreaSqMeters, ELightUnits::Lumens);
-	}
-	else
-	{
-		LightComponent.Intensity = NitsIntensity;
-		LightComponent.IntensityUnits = ELightUnits::Nits;
-	}
+	LightComponent.Intensity = NitsIntensity;
+	LightComponent.IntensityUnits = ELightUnits::Nits;
 
 	return true;
 }
@@ -250,51 +227,14 @@ bool UsdToUnreal::ConvertSphereLight(const pxr::UsdPrim& Prim, UPointLightCompon
 	const float UsdIntensity = UsdUtils::GetLightAttrValueWithInputsFallback<float>(SphereLight.GetIntensityAttr(), UsdTimeCode);
 	const float UsdExposure = UsdUtils::GetLightAttrValueWithInputsFallback<float>(SphereLight.GetExposureAttr(), UsdTimeCode);
 	const float UsdRadius = UsdUtils::GetLightAttrValueWithInputsFallback<float>(SphereLight.GetRadiusAttr(), UsdTimeCode);
-	const double UsdNormalized = UsdUtils::GetLightAttrValueWithInputsFallback<bool>(SphereLight.GetNormalizeAttr(), pxr::UsdTimeCode::Default().GetValue());
 	const double TreatAsPoint = UsdUtils::GetUsdValue<bool>(SphereLight.GetTreatAsPointAttr(), UsdTimeCode);
 
 	const float RadiusUE = UsdToUnreal::ConvertDistance(StageInfo, UsdRadius);
 	const double NitsIntensity = UsdToUnreal::ConvertLightIntensity(UsdIntensity, UsdExposure);
 
 	LightComponent.SourceRadius = RadiusUE;
-	if (UsdNormalized)
-	{
-		const float RadiusMeters = RadiusUE / 100.0f;
-		float SolidAngle = 1.0f;
-		float AreaSqMeters = 1.0f;
-
-		// Annoying, but historically the convention has been to call both this and ConvertLuxShapingAPI on the same prim/component
-		// as this handles the intensity part and that handles the cone angles. So here we need to redo some of the work that
-		// that function also does in order to also retrieve the angles for the unit conversion
-		if (USpotLightComponent* SpotLightComponent = Cast<USpotLightComponent>(&LightComponent))
-		{
-			pxr::UsdLuxShapingAPI ShapingAPI{Prim};
-			if (ShapingAPI)
-			{
-				const float UsdConeAngle = UsdUtils::GetLightAttrValueWithInputsFallback<float>(ShapingAPI.GetShapingConeAngleAttr(), UsdTimeCode);
-				const float UsdConeSoftness = UsdUtils::GetLightAttrValueWithInputsFallback<float>(ShapingAPI.GetShapingConeSoftnessAttr(), UsdTimeCode);
-				float InnerConeAngle = 0.0f;
-				float OuterConeAngle = 1.0f;
-				UsdToUnreal::ConvertSpotLightConeAngles(UsdConeAngle, UsdConeSoftness, OuterConeAngle, InnerConeAngle);
-
-				SolidAngle = UsdUtils::GetSpotLightSolidAngle(OuterConeAngle, InnerConeAngle);
-				AreaSqMeters = UsdUtils::GetSpotLightArea(RadiusMeters);
-			}
-		}
-		else // Just an actual point light
-		{
-			SolidAngle = UsdUtils::GetPointLightSolidAngle();
-			AreaSqMeters = UsdUtils::GetPointLightArea(RadiusMeters);
-		}
-
-		LightComponent.IntensityUnits = ELightUnits::Lumens;
-		LightComponent.Intensity = UsdUtils::ConvertIntensityFromNits(NitsIntensity, SolidAngle, AreaSqMeters, ELightUnits::Lumens);
-	}
-	else
-	{
-		LightComponent.Intensity = NitsIntensity;
-		LightComponent.IntensityUnits = ELightUnits::Nits;
-	}
+	LightComponent.Intensity = NitsIntensity;
+	LightComponent.IntensityUnits = ELightUnits::Nits;
 
 	return true;
 }
@@ -459,36 +399,6 @@ bool UnrealToUsd::ConvertLightComponent(const ULightComponentBase& LightComponen
 
 	if (const ULightComponent* DerivedLightComponent = Cast<ULightComponent>(&LightComponent))
 	{
-		// If we're using a normalized intensity unit, author inputs:normalize=true, but otherwise clear the 
-		// authored value (the default is false).
-		//
-		// We only read/write the default opinions for now as we don't support animating which
-		// light unit we're using anyway.
-		ELightUnits Units = DerivedLightComponent->GetLightUnits();
-		switch (Units)
-		{
-			case ELightUnits::Candelas:
-			case ELightUnits::EV:
-			case ELightUnits::Lumens:
-			case ELightUnits::Unitless:
-			{
-				if (pxr::UsdAttribute Attr = LightAPI.CreateNormalizeAttr())
-				{
-					Attr.Set<bool>(true, UsdUtils::GetDefaultTimeCode());
-					UsdUtils::NotifyIfOverriddenOpinion(Attr);
-				}
-				break;
-			}
-			default:
-			{
-				if (pxr::UsdAttribute Attr = LightAPI.GetNormalizeAttr())
-				{
-					Attr.ClearDefault();
-				}
-				break;
-			}
-		}
-
 		if (pxr::UsdAttribute Attr = LightAPI.CreateEnableColorTemperatureAttr())
 		{
 			Attr.Set<bool>(DerivedLightComponent->bUseTemperature, UsdTimeCode);
diff --git a/Engine/Plugins/Runtime/USDCore/Source/USDUtilities/Private/USDPrimConversion.cpp b/Engine/Plugins/Runtime/USDCore/Source/USDUtilities/Private/USDPrimConversion.cpp
--- a/Engine/Plugins/Runtime/USDCore/Source/USDUtilities/Private/USDPrimConversion.cpp
+++ b/Engine/Plugins/Runtime/USDCore/Source/USDUtilities/Private/USDPrimConversion.cpp
@@ -5895,8 +5895,6 @@ UnrealToUsd::FPropertyTrackWriter UnrealToUsd::CreatePropertyTrackWriter(
 			{
 				if (PropertyPath == UnrealIdentifiers::SourceRadiusPropertyName)
 				{
-					OutPropertyPathsToRefresh.Add(UnrealIdentifiers::IntensityPropertyName);
-
 					Attr = SphereLight.GetRadiusAttr();
 					if (Attr)
 					{
@@ -6038,8 +6036,6 @@ UnrealToUsd::FPropertyTrackWriter UnrealToUsd::CreatePropertyTrackWriter(
 					Attr = RectLight.GetWidthAttr();
 					if (Attr)
 					{
-						OutPropertyPathsToRefresh.Add(UnrealIdentifiers::IntensityPropertyName);
-
 						Result.FloatWriter = [Attr, StageInfo](float UEValue, double UsdTimeCode)
 						{
 							Attr.Set(UnrealToUsd::ConvertDistance(StageInfo, UEValue), UsdTimeCode);
@@ -6051,8 +6047,6 @@ UnrealToUsd::FPropertyTrackWriter UnrealToUsd::CreatePropertyTrackWriter(
 					Attr = RectLight.GetHeightAttr();
 					if (Attr)
 					{
-						OutPropertyPathsToRefresh.Add(UnrealIdentifiers::IntensityPropertyName);
-
 						Result.FloatWriter = [Attr, StageInfo](float UEValue, double UsdTimeCode)
 						{
 							Attr.Set(UnrealToUsd::ConvertDistance(StageInfo, UEValue), UsdTimeCode);
@@ -6104,8 +6098,6 @@ UnrealToUsd::FPropertyTrackWriter UnrealToUsd::CreatePropertyTrackWriter(
 					Attr = DiskLight.GetRadiusAttr();
 					if (Attr)
 					{
-						OutPropertyPathsToRefresh.Add(UnrealIdentifiers::IntensityPropertyName);
-
 						// Resync the other to match this one after we bake it, effectively always enforcing the UE rect light into a square shape
 						OutPropertyPathsToRefresh.Add(
 							PropertyPath == UnrealIdentifiers::SourceWidthPropertyName ? UnrealIdentifiers::SourceHeightPropertyName
@@ -6705,51 +6697,10 @@ TArray<UE::FUsdAttribute> UsdUtils::GetAttributesForProperty(const UE::FUsdPrim&
 	// Light attributes
 	else if (PropertyPath == IntensityPropertyName)
 	{
-		if (UsdPrim.IsA<pxr::UsdLuxRectLight>())
-		{
-			return {
-				UE::FUsdAttribute{UsdPrim.GetAttribute(pxr::UsdLuxTokens->inputsIntensity)},
-				UE::FUsdAttribute{UsdPrim.GetAttribute(pxr::UsdLuxTokens->inputsExposure)},
-				UE::FUsdAttribute{UsdPrim.GetAttribute(pxr::UsdLuxTokens->inputsWidth)},
-				UE::FUsdAttribute{UsdPrim.GetAttribute(pxr::UsdLuxTokens->inputsHeight)}
-			};
-		}
-		else if (UsdPrim.IsA<pxr::UsdLuxDiskLight>())
-		{
-			return {
-				UE::FUsdAttribute{UsdPrim.GetAttribute(pxr::UsdLuxTokens->inputsIntensity)},
-				UE::FUsdAttribute{UsdPrim.GetAttribute(pxr::UsdLuxTokens->inputsExposure)},
-				UE::FUsdAttribute{UsdPrim.GetAttribute(pxr::UsdLuxTokens->inputsRadius)}
-			};
-		}
-		else if (UsdPrim.IsA<pxr::UsdLuxDistantLight>() || UsdPrim.IsA<pxr::UsdLuxDomeLight>() || UsdPrim.IsA<pxr::UsdLuxDomeLight_1>())
-		{
-			return {
-				UE::FUsdAttribute{UsdPrim.GetAttribute(pxr::UsdLuxTokens->inputsIntensity)},
-				UE::FUsdAttribute{UsdPrim.GetAttribute(pxr::UsdLuxTokens->inputsExposure)}
-			};
-		}
-		else if (UsdPrim.IsA<pxr::UsdLuxSphereLight>())
-		{
-			if (UsdPrim.HasAPI<pxr::UsdLuxShapingAPI>())
-			{
-				return {
-					UE::FUsdAttribute{UsdPrim.GetAttribute(pxr::UsdLuxTokens->inputsIntensity)},
-					UE::FUsdAttribute{UsdPrim.GetAttribute(pxr::UsdLuxTokens->inputsExposure)},
-					UE::FUsdAttribute{UsdPrim.GetAttribute(pxr::UsdLuxTokens->inputsRadius)},
-					UE::FUsdAttribute{UsdPrim.GetAttribute(pxr::UsdLuxTokens->inputsShapingConeAngle)},
-					UE::FUsdAttribute{UsdPrim.GetAttribute(pxr::UsdLuxTokens->inputsShapingConeSoftness)}
-				};
-			}
-			else
-			{
-				return {
-					UE::FUsdAttribute{UsdPrim.GetAttribute(pxr::UsdLuxTokens->inputsIntensity)},
-					UE::FUsdAttribute{UsdPrim.GetAttribute(pxr::UsdLuxTokens->inputsExposure)},
-					UE::FUsdAttribute{UsdPrim.GetAttribute(pxr::UsdLuxTokens->inputsRadius)}
-				};
-			}
-		}
+		return {
+			UE::FUsdAttribute{UsdPrim.GetAttribute(pxr::UsdLuxTokens->inputsIntensity)},
+			UE::FUsdAttribute{UsdPrim.GetAttribute(pxr::UsdLuxTokens->inputsExposure)}
+		};
 	}
 	else if (PropertyPath == LightColorPropertyName)
 	{
@@ -6894,6 +6845,11 @@ TArray<FName> UsdUtils::GetPropertiesForAttribute(const UE::FUsdPrim& Prim, cons
 		{pxr::UsdLuxTokens->inputsIntensity, 				{IntensityPropertyName}},
 		{pxr::UsdLuxTokens->inputsExposure, 				{IntensityPropertyName}}, // In USD true intensity is just a function of "intensity" and "exposure"
 		{pxr::UsdLuxTokens->inputsColor, 					{LightColorPropertyName}},
+		{pxr::UsdLuxTokens->inputsWidth, 					{SourceWidthPropertyName}},
+		{pxr::UsdLuxTokens->inputsHeight, 					{SourceHeightPropertyName}},
+		{pxr::UsdLuxTokens->inputsAngle, 					{LightSourceAnglePropertyName}},
+		{pxr::UsdLuxTokens->inputsShapingConeAngle, 		{OuterConeAnglePropertyName}},
+		{pxr::UsdLuxTokens->inputsShapingConeSoftness, 		{InnerConeAnglePropertyName}},
 		{pxr::UsdLuxTokens->inputsColorTemperature, 		{TemperaturePropertyName}},
 		{pxr::UsdLuxTokens->inputsEnableColorTemperature, 	{UseTemperaturePropertyName}},
 	};
@@ -6904,51 +6860,16 @@ TArray<FName> UsdUtils::GetPropertiesForAttribute(const UE::FUsdPrim& Prim, cons
 		return Iter->second;
 	}
 
-	// Check to see if it's one of the edge cases that depends on the prim type.
+	// Check to see if it's one of the edge cases that depends on the prim type
 	if (AttrToken == pxr::UsdLuxTokens->inputsRadius)
 	{
 		if (UsdPrim.IsA<pxr::UsdLuxSphereLight>())
 		{
-			return {SourceRadiusPropertyName, IntensityPropertyName};
+			return {SourceRadiusPropertyName};
 		}
 		else if (UsdPrim.IsA<pxr::UsdLuxDiskLight>())
 		{
-			return {SourceWidthPropertyName, SourceHeightPropertyName, IntensityPropertyName};
-		}
-	}
-	else if (AttrToken == pxr::UsdLuxTokens->inputsWidth)
-	{
-		if (UsdPrim.IsA<pxr::UsdLuxRectLight>())
-		{
-			return {SourceWidthPropertyName, IntensityPropertyName};
-		}
-	}
-	else if (AttrToken == pxr::UsdLuxTokens->inputsHeight)
-	{
-		if (UsdPrim.IsA<pxr::UsdLuxRectLight>())
-		{
-			return {SourceHeightPropertyName, IntensityPropertyName};
-		}
-	}
-	else if (AttrToken == pxr::UsdLuxTokens->inputsAngle)
-	{
-		if (UsdPrim.IsA<pxr::UsdLuxDistantLight>())
-		{
-			return {LightSourceAnglePropertyName};
-		}
-	}
-	else if (AttrToken == pxr::UsdLuxTokens->inputsShapingConeAngle)
-	{
-		if (UsdPrim.IsA<pxr::UsdLuxSphereLight>() && UsdPrim.HasAPI<pxr::UsdLuxShapingAPI>())
-		{
-			return {OuterConeAnglePropertyName, IntensityPropertyName};
-		}
-	}
-	else if (AttrToken == pxr::UsdLuxTokens->inputsShapingConeSoftness)
-	{
-		if (UsdPrim.IsA<pxr::UsdLuxSphereLight>() && UsdPrim.HasAPI<pxr::UsdLuxShapingAPI>())
-		{
-			return {InnerConeAnglePropertyName, IntensityPropertyName};
+			return {SourceWidthPropertyName, SourceHeightPropertyName};
 		}
 	}
 
